Sequencer: Initial version of "Stitch Track".
The stitch track is a animation pose producing track using motion matching that allows the user to place a section referencing a motion matching database and a target pose. When the section activates, it will use motion matching to attempt to find an animation to connect the actor's current transform and pose history/trajectory with the target pose and transform, and then will play that back with root motion if selected. This can be used to create more seamless transitions from gameplay into the scene.

As part of this change, I've added a 'pose history' attribute and modified the anim bp node and anim next versions of the pose history node/trait to write this attribute into their output. This allows the sequencer-created evaluation task to access this pose history in order to do its computation.

Currently this track only does the motion matching calculation once at the initial evaluation of the section. However we may desire to change this to update more frequently and re-run the calculation.

In addition we may desire to generalize this track to allow the user to key a trajectory and potentially poses along the section to utilize motion matching in more ways during a scene.

Currently all of this functionality is gated behind a 'Pose Search Tracks' plugin that itself relies on the Pose Search, Sequencer
Anim Mixer, and Anim Next plugins.

diff --git a/Engine/Plugins/Animation/PoseSearch/Source/Runtime/Private/AnimNode_PoseSearchHistoryCollector.cpp b/Engine/Plugins/Animation/PoseSearch/Source/Runtime/Private/AnimNode_PoseSearchHistoryCollector.cpp
--- a/Engine/Plugins/Animation/PoseSearch/Source/Runtime/Private/AnimNode_PoseSearchHistoryCollector.cpp
+++ b/Engine/Plugins/Animation/PoseSearch/Source/Runtime/Private/AnimNode_PoseSearchHistoryCollector.cpp
@@ -6,6 +6,7 @@
 #include "Animation/AnimStats.h"
 #include "Engine/SkinnedAsset.h"
 #include "PoseSearch/PoseSearchInteractionValidator.h"
+#include "PoseSearch/PoseSearchHistoryAttribute.h"
 
 #include UE_INLINE_GENERATED_CPP_BY_NAME(AnimNode_PoseSearchHistoryCollector)
 
@@ -159,6 +160,13 @@ void FAnimNode_PoseSearchHistoryCollector::Evaluate_AnyThread(FPoseContext& Outp
 		RootBoneRecoveryTime, RootBoneTranslationRecoveryRatio, RootBoneRotationRecoveryRatio, bNeedsReset, bCacheBones, 
 		RequiredBones, Output.Curve, MakeConstArrayView(CollectedCurves));
 
+	// Fill in the pose history into a custom attribute for access downstream
+	if (FPoseHistoryAnimationAttribute* PoseHistoryAttribute = Output.CustomAttributes.FindOrAdd<FPoseHistoryAnimationAttribute>(UE::PoseSearch::PoseHistoryAttributeId))
+	{
+		PoseHistoryAttribute->PoseHistory = &GetPoseHistory();
+		PoseHistoryAttribute->ScopeObject = Output.GetAnimInstanceObject();
+	}
+
 	bCacheBones = false;
 
 #if ENABLE_DRAW_DEBUG && ENABLE_ANIM_DEBUG
@@ -221,6 +229,13 @@ void FAnimNode_PoseSearchComponentSpaceHistoryCollector::EvaluateComponentSpace_
 	GetPoseHistory().EvaluateComponentSpace_AnyThread(Output.AnimInstanceProxy->GetDeltaSeconds(), ComponentSpacePoseProvider, bStoreScales, 
 		RootBoneRecoveryTime, RootBoneTranslationRecoveryRatio, RootBoneRotationRecoveryRatio, bNeedsReset, bCacheBones,
 		RequiredBones, Output.Curve, MakeConstArrayView(CollectedCurves));
+
+	// Fill in the pose history into a custom attribute for access downstream
+	if (FPoseHistoryAnimationAttribute* PoseHistoryAttribute = Output.CustomAttributes.FindOrAdd<FPoseHistoryAnimationAttribute>(UE::PoseSearch::PoseHistoryAttributeId))
+	{
+		PoseHistoryAttribute->PoseHistory = &GetPoseHistory();
+		PoseHistoryAttribute->ScopeObject = Output.GetAnimInstanceObject();
+	}
 	
 	bCacheBones = false;
 
diff --git a/Engine/Plugins/Animation/PoseSearch/Source/Runtime/Private/PoseSearchHistoryAttribute.cpp b/Engine/Plugins/Animation/PoseSearch/Source/Runtime/Private/PoseSearchHistoryAttribute.cpp
new file mode 100644
--- /dev/null
+++ b/Engine/Plugins/Animation/PoseSearch/Source/Runtime/Private/PoseSearchHistoryAttribute.cpp
@@ -0,0 +1,8 @@
+// Copyright Epic Games, Inc. All Rights Reserved.
+
+#include "PoseSearch/PoseSearchHistoryAttribute.h"
+
+namespace UE::PoseSearch
+{
+	const UE::Anim::FAttributeId PoseHistoryAttributeId(TEXT("PoseHistoryAttribute"), FCompactPoseBoneIndex(0));
+}
\ No newline at end of file
diff --git a/Engine/Plugins/Animation/PoseSearch/Source/Runtime/Private/PoseSearchModule.cpp b/Engine/Plugins/Animation/PoseSearch/Source/Runtime/Private/PoseSearchModule.cpp
--- a/Engine/Plugins/Animation/PoseSearch/Source/Runtime/Private/PoseSearchModule.cpp
+++ b/Engine/Plugins/Animation/PoseSearch/Source/Runtime/Private/PoseSearchModule.cpp
@@ -10,6 +10,8 @@
 #include "PoseSearch/PoseSearchSchema.h"
 #include "PoseSearch/PoseSearchFeatureChannel_PermutationTime.h"
 #include "RewindDebuggerRuntimeInterface/IRewindDebuggerRuntimeExtension.h"
+#include "Animation/AttributeTypes.h"
+#include "PoseSearch/PoseSearchHistoryAttribute.h"
 
 class FRewindDebuggerPoseSearchRuntime : public IRewindDebuggerRuntimeExtension 
 {
@@ -34,12 +36,16 @@ public:
 	{
 		IModularFeatures::Get().RegisterModularFeature(UE::Anim::IPoseSearchProvider::GetModularFeatureName(), this);
 		IModularFeatures::Get().RegisterModularFeature(IRewindDebuggerRuntimeExtension::ModularFeatureName, &RewindDebuggerPoseSearchRuntime);
+
+		UE::Anim::AttributeTypes::RegisterType<FPoseHistoryAnimationAttribute>();
 	}
 
 	virtual void ShutdownModule() override
 	{
 		IModularFeatures::Get().UnregisterModularFeature(UE::Anim::IPoseSearchProvider::GetModularFeatureName(), this);
 		IModularFeatures::Get().UnregisterModularFeature(IRewindDebuggerRuntimeExtension::ModularFeatureName, &RewindDebuggerPoseSearchRuntime);
+
+		UE::Anim::AttributeTypes::UnregisterType<FPoseHistoryAnimationAttribute>();
 	}
 
 	// IPoseSearchProvider
diff --git a/Engine/Plugins/Animation/PoseSearch/Source/Runtime/Public/PoseSearch/PoseSearchHistoryAttribute.h b/Engine/Plugins/Animation/PoseSearch/Source/Runtime/Public/PoseSearch/PoseSearchHistoryAttribute.h
new file mode 100644
--- /dev/null
+++ b/Engine/Plugins/Animation/PoseSearch/Source/Runtime/Public/PoseSearch/PoseSearchHistoryAttribute.h
@@ -0,0 +1,46 @@
+// Copyright Epic Games, Inc. All Rights Reserved.
+
+#pragma once
+
+#include "PoseSearchHistory.h"
+#include "Animation/AttributesContainer.h"
+#include "Animation/AttributeTraits.h"
+#include "PoseSearchHistoryAttribute.generated.h"
+
+namespace UE::PoseSearch
+{
+	POSESEARCH_API extern const UE::Anim::FAttributeId PoseHistoryAttributeId;
+}
+
+/** Attribute type containing a const pointer to pose history for use by downstream systems. Uses a TWeakObjectPtr for scope checks.*/
+USTRUCT(Experimental)
+struct POSESEARCH_API FPoseHistoryAnimationAttribute
+{
+	GENERATED_BODY()
+
+	const UE::PoseSearch::IPoseHistory* PoseHistory = nullptr;
+	TWeakObjectPtr<const UObject> ScopeObject;
+
+	bool IsValid() const
+	{
+		return ScopeObject.IsValid() && PoseHistory != nullptr;
+	}
+};
+
+inline POSESEARCH_API uint32 GetTypeHash(const FPoseHistoryAnimationAttribute& Key)
+{
+	return HashCombine(GetTypeHash(Key.ScopeObject), GetTypeHash(Key.PoseHistory));
+}
+
+namespace UE::Anim
+{
+	/** Pose history attribute is not blend-able by default */
+	template<>
+	struct POSESEARCH_API TAttributeTypeTraits<FPoseHistoryAnimationAttribute> : public TAttributeTypeTraitsBase<FPoseHistoryAnimationAttribute>
+	{
+		enum
+		{
+			IsBlendable = false,
+		};
+	};
+}
\ No newline at end of file
diff --git a/Engine/Plugins/Animation/PoseSearch/Source/Runtime/Public/PoseSearch/PoseSearchIndex.h b/Engine/Plugins/Animation/PoseSearch/Source/Runtime/Public/PoseSearch/PoseSearchIndex.h
--- a/Engine/Plugins/Animation/PoseSearch/Source/Runtime/Public/PoseSearch/PoseSearchIndex.h
+++ b/Engine/Plugins/Animation/PoseSearch/Source/Runtime/Public/PoseSearch/PoseSearchIndex.h
@@ -568,7 +568,7 @@ struct FSearchIndex : public FSearchIndexBase
 	FSearchIndex& operator=(FSearchIndex&& Other);
 
 	void Reset();
-	TConstArrayView<float> GetPoseValues(int32 PoseIdx) const;
+	POSESEARCH_API TConstArrayView<float> GetPoseValues(int32 PoseIdx) const;
 	TConstArrayView<float> GetReconstructedPoseValues(int32 PoseIdx, TArrayView<float> BufferUsedForReconstruction) const;
 	int32 GetNumDimensions() const;
 	int32 GetNumberOfPrincipalComponents() const;
diff --git a/Engine/Plugins/Experimental/AnimNextPoseSearch/Source/AnimNextPoseSearch/Private/HistoryCollectorTrait.cpp b/Engine/Plugins/Experimental/AnimNextPoseSearch/Source/AnimNextPoseSearch/Private/HistoryCollectorTrait.cpp
--- a/Engine/Plugins/Experimental/AnimNextPoseSearch/Source/AnimNextPoseSearch/Private/HistoryCollectorTrait.cpp
+++ b/Engine/Plugins/Experimental/AnimNextPoseSearch/Source/AnimNextPoseSearch/Private/HistoryCollectorTrait.cpp
@@ -6,6 +6,7 @@
 #include "AnimationRuntime.h"
 #include "EvaluationVM/EvaluationVM.h"
 #include "GenerationTools.h"
+#include "PoseSearch/PoseSearchHistoryAttribute.h"
 
 namespace UE::AnimNext
 {
@@ -83,7 +84,7 @@ void FAnimNextHistoryCollectorTask::Execute(UE::AnimNext::FEvaluationVM& VM) con
 
 	if (EnumHasAnyFlags(VM.GetFlags(), EEvaluationFlags::Bones))
 	{
-		if (const TUniquePtr<FKeyframeState>* Keyframe = VM.PeekValue<TUniquePtr<FKeyframeState>>(KEYFRAME_STACK_NAME, 0))
+		if (TUniquePtr<FKeyframeState>* Keyframe = VM.PeekValueMutable<TUniquePtr<FKeyframeState>>(KEYFRAME_STACK_NAME, 0))
 		{
 			bool bNeedsReset = false;
 			bool bCacheBones = false;
@@ -128,6 +129,13 @@ void FAnimNextHistoryCollectorTask::Execute(UE::AnimNext::FEvaluationVM& VM) con
 			PoseHistory->EvaluateComponentSpace_AnyThread(DeltaTime, ComponentSpacePoseProvider, bStoreScales,
 				SharedData->RootBoneRecoveryTime, SharedData->RootBoneTranslationRecoveryRatio, SharedData->RootBoneRotationRecoveryRatio,
 				bNeedsReset, bCacheBones, BoneIndicesWithParents);
+
+			// Push pose history onto the custom attribute for downstream tasks
+			if (FPoseHistoryAnimationAttribute* PoseHistoryAttribute = (*Keyframe)->Attributes.FindOrAdd<FPoseHistoryAnimationAttribute>(UE::PoseSearch::PoseHistoryAttributeId))
+			{
+				PoseHistoryAttribute->PoseHistory = PoseHistory;
+				PoseHistoryAttribute->ScopeObject = HostObject;
+			}
 		}
 	}
 }
@@ -190,6 +198,7 @@ namespace UE::AnimNext
 		Task.SharedData = SharedData;
 		Task.bStoreScales = bStoreScales;
 		Task.DeltaTime = InstanceData->DeltaTime;
+		Task.HostObject = Context.GetHostObject();
 		Context.AppendTask(Task);
 	}
 
diff --git a/Engine/Plugins/Experimental/AnimNextPoseSearch/Source/AnimNextPoseSearch/Private/HistoryCollectorTrait.h b/Engine/Plugins/Experimental/AnimNextPoseSearch/Source/AnimNextPoseSearch/Private/HistoryCollectorTrait.h
--- a/Engine/Plugins/Experimental/AnimNextPoseSearch/Source/AnimNextPoseSearch/Private/HistoryCollectorTrait.h
+++ b/Engine/Plugins/Experimental/AnimNextPoseSearch/Source/AnimNextPoseSearch/Private/HistoryCollectorTrait.h
@@ -9,6 +9,7 @@
 #include "TraitCore/Trait.h"
 #include "TraitInterfaces/IEvaluate.h"
 #include "TraitInterfaces/IUpdate.h"
+#include "EvaluationVM/EvaluationVM.h"
 #include "HistoryCollectorTrait.generated.h"
 
 USTRUCT(meta = (DisplayName = "Pose History"))
@@ -121,6 +122,7 @@ struct FAnimNextHistoryCollectorTask : public FAnimNextEvaluationTask
 	const FAnimNextHistoryCollectorTraitSharedData* SharedData = nullptr;
 	bool bStoreScales = false;
 	float DeltaTime = 0.f;
+	TWeakObjectPtr<const UObject> HostObject;
 };
 
 namespace UE::AnimNext
diff --git a/Engine/Plugins/Experimental/MovieSceneAnimMixer/Source/MovieSceneAnimMixer/Private/Systems/MovieSceneAnimMixerSystem.cpp b/Engine/Plugins/Experimental/MovieSceneAnimMixer/Source/MovieSceneAnimMixer/Private/Systems/MovieSceneAnimMixerSystem.cpp
--- a/Engine/Plugins/Experimental/MovieSceneAnimMixer/Source/MovieSceneAnimMixer/Private/Systems/MovieSceneAnimMixerSystem.cpp
+++ b/Engine/Plugins/Experimental/MovieSceneAnimMixer/Source/MovieSceneAnimMixer/Private/Systems/MovieSceneAnimMixerSystem.cpp
@@ -153,7 +153,6 @@ UMovieSceneAnimMixerSystem::UMovieSceneAnimMixerSystem(const FObjectInitializer&
 	{
 		FBuiltInComponentTypes* BuiltInComponents = FBuiltInComponentTypes::Get();
 
-		DefineImplicitPrerequisite(UMovieSceneGenericBoundObjectInstantiator::StaticClass(), GetClass());
 		DefineImplicitPrerequisite(UMovieSceneGenericBoundObjectInstantiator::StaticClass(), GetClass());
 		DefineImplicitPrerequisite(UMovieSceneBoundSceneComponentInstantiator::StaticClass(), GetClass());
 		DefineImplicitPrerequisite(UWeightAndEasingEvaluatorSystem::StaticClass(), GetClass());
diff --git a/Engine/Plugins/Experimental/MovieScenePoseSearchTracks/MovieScenePoseSearchTracks.uplugin b/Engine/Plugins/Experimental/MovieScenePoseSearchTracks/MovieScenePoseSearchTracks.uplugin
new file mode 100644
--- /dev/null
+++ b/Engine/Plugins/Experimental/MovieScenePoseSearchTracks/MovieScenePoseSearchTracks.uplugin
@@ -0,0 +1,42 @@
+{
+	"FileVersion" : 1,
+	"Version" : 1,
+	"VersionName" : "0.1",
+	"FriendlyName" : "Movie Scene Pose Search Tracks",
+	"Description" : "Sequencer pose search tracks using the Anim Mixer",
+	"Category" : "Animation",
+	"CreatedBy" : "Epic Games, Inc.",
+	"CreatedByURL" : "https://epicgames.com",
+	"DocsURL" : "",
+	"MarketplaceURL" : "",
+	"SupportURL" : "",
+	"EnabledByDefault" : false,
+	"CanContainContent" : true,
+	"IsBetaVersion": false,
+	"IsExperimentalVersion": true,
+	"Installed" : false,
+	"Plugins":
+	[
+		{
+			"Name": "MovieSceneAnimMixer",
+			"Enabled": true
+		},
+		{
+			"Name": "AnimNextPoseSearch",
+			"Enabled": true
+		}
+	],
+	"Modules" :
+	[
+		{
+			"Name" : "MovieScenePoseSearchTracks",
+			"Type" : "Runtime",
+			"LoadingPhase" : "Default"
+		},
+		{
+			"Name" : "MovieScenePoseSearchTracksEditor",
+			"Type" : "Editor",
+			"LoadingPhase" : "Default"
+		}
+	]
+}
diff --git a/Engine/Plugins/Experimental/MovieScenePoseSearchTracks/Source/MovieScenePoseSearchTracks/MovieScenePoseSearchTracks.Build.cs b/Engine/Plugins/Experimental/MovieScenePoseSearchTracks/Source/MovieScenePoseSearchTracks/MovieScenePoseSearchTracks.Build.cs
new file mode 100644
--- /dev/null
+++ b/Engine/Plugins/Experimental/MovieScenePoseSearchTracks/Source/MovieScenePoseSearchTracks/MovieScenePoseSearchTracks.Build.cs
@@ -0,0 +1,27 @@
+// Copyright Epic Games, Inc. All Rights Reserved.
+
+namespace UnrealBuildTool.Rules
+{
+	public class MovieScenePoseSearchTracks : ModuleRules
+	{
+		public MovieScenePoseSearchTracks(ReadOnlyTargetRules Target) : base(Target)
+		{
+			PublicDependencyModuleNames.AddRange(
+				new string[]
+				{
+					"Core",
+					"CoreUObject",
+					"Engine",
+					"MovieScene",
+					"MovieSceneTracks",
+					"MovieSceneAnimMixer",
+					"AnimNext",
+					"AnimNextPoseSearch",
+					"AnimNextAnimGraph",
+					"AnimGraphRuntime",
+					"PoseSearch"
+				}
+			);
+		}
+	}
+}
\ No newline at end of file
diff --git a/Engine/Plugins/Experimental/MovieScenePoseSearchTracks/Source/MovieScenePoseSearchTracks/Private/MovieScenePoseSearchTracksModule.cpp b/Engine/Plugins/Experimental/MovieScenePoseSearchTracks/Source/MovieScenePoseSearchTracks/Private/MovieScenePoseSearchTracksModule.cpp
new file mode 100644
--- /dev/null
+++ b/Engine/Plugins/Experimental/MovieScenePoseSearchTracks/Source/MovieScenePoseSearchTracks/Private/MovieScenePoseSearchTracksModule.cpp
@@ -0,0 +1,26 @@
+// Copyright Epic Games, Inc. All Rights Reserved.
+
+#include "MovieScenePoseSearchTracksModule.h"
+#include "Modules/ModuleInterface.h"
+#include "Modules/ModuleManager.h"
+#include "MovieSceneAnimMixerModule.h"
+#include "PoseSearch/PoseHistoryProvider.h"
+#include "PoseSearch/PoseSearchHistory.h"
+#include "Animation/AnimNodeBase.h"
+#include "EvaluationVM/EvaluationVM.h"
+
+namespace UE::MovieScene
+{
+
+void FMovieScenePoseSearchTracksModule::StartupModule()
+{
+}
+
+void FMovieScenePoseSearchTracksModule::ShutdownModule()
+{
+
+}
+
+}
+
+IMPLEMENT_MODULE(UE::MovieScene::FMovieScenePoseSearchTracksModule, MovieScenePoseSearchTracks)
diff --git a/Engine/Plugins/Experimental/MovieScenePoseSearchTracks/Source/MovieScenePoseSearchTracks/Private/MovieScenePoseSearchTracksModule.h b/Engine/Plugins/Experimental/MovieScenePoseSearchTracks/Source/MovieScenePoseSearchTracks/Private/MovieScenePoseSearchTracksModule.h
new file mode 100644
--- /dev/null
+++ b/Engine/Plugins/Experimental/MovieScenePoseSearchTracks/Source/MovieScenePoseSearchTracks/Private/MovieScenePoseSearchTracksModule.h
@@ -0,0 +1,25 @@
+// Copyright Epic Games, Inc. All Rights Reserved.
+
+#pragma once
+
+#include "CoreMinimal.h"
+#include "Modules/ModuleInterface.h"
+
+struct FAnimationUpdateContext;
+
+namespace UE::AnimNext
+{
+	struct FEvaluationVM;
+}
+
+namespace UE::MovieScene
+{
+
+class FMovieScenePoseSearchTracksModule : public IModuleInterface
+{
+	// IModuleInterface interface
+	virtual void StartupModule() override;
+	virtual void ShutdownModule() override;
+};
+
+}
diff --git a/Engine/Plugins/Experimental/MovieScenePoseSearchTracks/Source/MovieScenePoseSearchTracks/Private/PoseSearchTracksComponentTypes.cpp b/Engine/Plugins/Experimental/MovieScenePoseSearchTracks/Source/MovieScenePoseSearchTracks/Private/PoseSearchTracksComponentTypes.cpp
new file mode 100644
--- /dev/null
+++ b/Engine/Plugins/Experimental/MovieScenePoseSearchTracks/Source/MovieScenePoseSearchTracks/Private/PoseSearchTracksComponentTypes.cpp
@@ -0,0 +1,39 @@
+// Copyright Epic Games, Inc. All Rights Reserved.
+
+#include "PoseSearchTracksComponentTypes.h"
+#include "EntitySystem/MovieSceneComponentRegistry.h"
+#include "EntitySystem/MovieSceneEntityFactoryTemplates.h"
+#include "EntitySystem/MovieSceneEntitySystemLinker.h"
+
+namespace UE::MovieScene
+{
+	static TUniquePtr<FPoseSearchTracksComponentTypes> GPoseSearchTracksComponentTypes;
+	static bool GPoseSearchTracksComponentTypesDestroyed = false;
+
+	FPoseSearchTracksComponentTypes* FPoseSearchTracksComponentTypes::Get()
+	{
+		if (!GPoseSearchTracksComponentTypes.IsValid())
+		{
+			check(!GPoseSearchTracksComponentTypesDestroyed);
+			GPoseSearchTracksComponentTypes.Reset(new FPoseSearchTracksComponentTypes);
+		}
+		return GPoseSearchTracksComponentTypes.Get();
+	}
+
+	void FPoseSearchTracksComponentTypes::Destroy()
+	{
+		GPoseSearchTracksComponentTypes.Reset();
+		GPoseSearchTracksComponentTypesDestroyed = true;
+	}
+
+	FPoseSearchTracksComponentTypes::FPoseSearchTracksComponentTypes()
+	{
+		FComponentRegistry* ComponentRegistry = UMovieSceneEntitySystemLinker::GetComponents();
+
+		ComponentRegistry->NewComponentType(&StitchAnim, TEXT("Stitch Animation"));
+
+		ComponentRegistry->Factories.DuplicateChildComponent(StitchAnim);
+
+	}
+
+}
\ No newline at end of file
diff --git a/Engine/Plugins/Experimental/MovieScenePoseSearchTracks/Source/MovieScenePoseSearchTracks/Private/Sections/MovieSceneStitchAnimSection.cpp b/Engine/Plugins/Experimental/MovieScenePoseSearchTracks/Source/MovieScenePoseSearchTracks/Private/Sections/MovieSceneStitchAnimSection.cpp
new file mode 100644
--- /dev/null
+++ b/Engine/Plugins/Experimental/MovieScenePoseSearchTracks/Source/MovieScenePoseSearchTracks/Private/Sections/MovieSceneStitchAnimSection.cpp
@@ -0,0 +1,284 @@
+// Copyright Epic Games, Inc. All Rights Reserved.
+
+#include "Sections/MovieSceneStitchAnimSection.h"
+#include "Channels/MovieSceneChannelProxy.h"
+#include "Tracks/MovieSceneStitchAnimTrack.h"
+#include "PoseSearchTracksComponentTypes.h"
+#include "PoseSearch/PoseSearchHistory.h"
+#include "Animation/AnimPoseSearchProvider.h"
+#include "EntitySystem/BuiltInComponentTypes.h"
+#include "Systems/MovieSceneSkeletalAnimationSystem.h"
+#include "MovieSceneTracksComponentTypes.h"
+#include "AnimMixerComponentTypes.h"
+#include "PoseSearch/PoseSearchHistoryAttribute.h"
+#include "EvaluationVM/KeyframeState.h"
+#include "EvaluationVM/EvaluationVM.h"
+#include "PoseSearch/PoseSearchLibrary.h"
+#include "PoseSearch/PoseSearchDatabase.h"
+#include "PoseSearch/PoseSearchFeatureChannel_PermutationTime.h"
+#include "PoseSearch/PoseSearchSchema.h"
+#include "DecompressionTools.h"
+#include "Animation/AnimInstance.h"
+#include "Animation/AnimSequence.h"
+#include "Variants/MovieSceneTimeWarpVariantPayloads.h"
+#include "Animation/BuiltInAttributeTypes.h"
+
+#include UE_INLINE_GENERATED_CPP_BY_NAME(MovieSceneStitchAnimSection)
+
+#define LOCTEXT_NAMESPACE "MovieSceneStitchAnimSection"
+
+
+void FMovieSceneStitchAnimEvaluationTask::Execute(UE::AnimNext::FEvaluationVM& VM) const
+{
+	using namespace UE::AnimNext;
+	using namespace UE::MovieScene;
+	using namespace UE::PoseSearch;
+	using namespace UE::Anim;
+
+	const FAnimMixerComponentTypes* AnimMixerComponents = FAnimMixerComponentTypes::Get();
+
+	// We rely on the pose history having been stored first.
+	// TODO: Find a way to initialize a pose history here for the first time and store it in something the task can access (stitch anim system?)
+	// if not already present.
+	const IPoseHistory* PoseHistory = nullptr;
+	if (const TUniquePtr<FKeyframeState>* Keyframe = VM.PeekValue<TUniquePtr<FKeyframeState>>(KEYFRAME_STACK_NAME, 0))
+	{
+		const FPoseHistoryAnimationAttribute* PoseHistoryAttribute = (*Keyframe)->Attributes.Find<FPoseHistoryAnimationAttribute>(PoseHistoryAttributeId);
+		if (PoseHistoryAttribute && PoseHistoryAttribute->ScopeObject.IsValid())
+		{
+			PoseHistory = PoseHistoryAttribute->PoseHistory;
+		}
+	}
+
+	// If we haven't yet run a search and found an animation, do that
+	if (!MatchedAsset)
+	{
+		if (PoseHistory && StitchData.StitchDatabase && StitchData.TargetPoseAsset)
+		{
+			// TODO: Potentially I can use this if I swap to a model of repeatedly searching each frame
+				/*IPoseSearchProvider::FSearchPlayingAsset PlayingAsset;
+				PlayingAsset.Asset = GetAnimAsset();
+				PlayingAsset.AccumulatedTime = GetAccumulatedTime();*/
+
+			FPoseSearchContinuingProperties ContinuingProperties;
+			/*ContinuingProperties.PlayingAsset = PlayingAsset.Asset;
+			ContinuingProperties.PlayingAssetAccumulatedTime = PlayingAsset.AccumulatedTime;*/
+
+			FPoseSearchFutureProperties Future;
+			Future.Animation = StitchData.TargetPoseAsset;
+			Future.AnimationTime = StitchData.TargetAnimationTimeSeconds;
+			Future.IntervalTime = TimeToTarget;
+			const UObject* StitchDatabaseObj = StitchData.StitchDatabase.Get();
+
+			const UObject* AnimContext = ContextObject.Get();
+			const FSearchResult SearchResult = UPoseSearchLibrary::MotionMatch(MakeArrayView(&AnimContext, 1),
+				MakeArrayView(&DefaultRole, 1), MakeArrayView(&PoseHistory, 1), MakeArrayView(&StitchDatabaseObj, 1), ContinuingProperties, Future, FPoseSearchEvent());
+
+			if (const FSearchIndexAsset* SearchIndexAsset = SearchResult.GetSearchIndexAsset())
+			{
+				const UPoseSearchDatabase* Database = SearchResult.Database.Get();
+				if (const FPoseSearchDatabaseAnimationAssetBase* DatabaseAnimationAssetBase = Database->GetDatabaseAnimationAsset<FPoseSearchDatabaseAnimationAssetBase>(*SearchIndexAsset))
+				{
+					MatchedAsset = Cast<UAnimationAsset>(DatabaseAnimationAssetBase->GetAnimationAsset());
+					MatchedAssetTime = SearchResult.AssetTime;
+
+					// TODO: Add mirroring support
+					//ProviderResult.bMirrored = SearchIndexAsset->IsMirrored();
+
+					// figuring out the WantedPlayRate
+					MatchedAssetPlayRate = 1.f;
+					if (Future.Animation && Future.IntervalTime > 0.f)
+					{
+						if (const UPoseSearchFeatureChannel_PermutationTime* PermutationTimeChannel = Database->Schema->FindFirstChannelOfType<UPoseSearchFeatureChannel_PermutationTime>())
+						{
+							const FSearchIndex& SearchIndex = Database->GetSearchIndex();
+							if (!SearchIndex.IsValuesEmpty())
+							{
+								TConstArrayView<float> ResultData = Database->GetSearchIndex().GetPoseValues(SearchResult.PoseIdx);
+								const float ActualIntervalTime = PermutationTimeChannel->GetPermutationTime(ResultData);
+								MatchedAssetPlayRate = ActualIntervalTime / Future.IntervalTime;
+							}
+						}
+					}
+				}
+			}
+		}
+	}
+
+	if (!MatchedAsset)
+	{
+		constexpr bool bIsAdditive = false;
+		FKeyframeState ReferenceKeyframe = VM.MakeReferenceKeyframe(bIsAdditive);
+		VM.PushValue(KEYFRAME_STACK_NAME, MakeUnique<FKeyframeState>(MoveTemp(ReferenceKeyframe)));
+		return;
+	}
+
+	FKeyframeState Keyframe = VM.MakeUninitializedKeyframe(false);
+	// TODO: For now assume this asset is a UAnimSequence
+
+	UAnimSequence* AnimSequence = Cast<UAnimSequence>(MatchedAsset);
+	if (!AnimSequence)
+	{
+		constexpr bool bIsAdditive = false;
+		FKeyframeState ReferenceKeyframe = VM.MakeReferenceKeyframe(bIsAdditive);
+		VM.PushValue(KEYFRAME_STACK_NAME, MakeUnique<FKeyframeState>(MoveTemp(ReferenceKeyframe)));
+		return;
+	}
+	const bool bUseRawData = FDecompressionTools::ShouldUseRawData(AnimSequence, Keyframe.Pose);
+
+	float AnimationTime = MatchedAssetTime + ((CurrentTime - InitialTime) * MatchedAssetPlayRate);
+	float PreviousAnimationTime = AnimationTime - ((CurrentTime - PreviousTime) * MatchedAssetPlayRate);
+
+	// Calculate the root transform. 
+	// We do this by aligning and projecting the matched animation asset twice,
+	// first to the 'initial' transform, and then to the 'target' transform, and blending the results based on time proximity.
+
+	FTransform InitialAnimSpaceRootTransform = AnimSequence->ExtractRootTrackTransform(FAnimExtractContext((double)MatchedAssetTime), nullptr);
+	FTransform CurrentAnimSpaceRootTransform = AnimSequence->ExtractRootTrackTransform(FAnimExtractContext((double)AnimationTime), nullptr);
+	FTransform TargetAnimSpaceRootTransform = AnimSequence->ExtractRootTrackTransform(FAnimExtractContext((double)(AnimationTime + (TimeToTarget * MatchedAssetPlayRate))), nullptr);
+
+	if (InverseMeshToActorRotation)
+	{
+		InitialAnimSpaceRootTransform.SetTranslation(InverseMeshToActorRotation.GetValue() * InitialAnimSpaceRootTransform.GetTranslation());
+		CurrentAnimSpaceRootTransform.SetTranslation(InverseMeshToActorRotation.GetValue() * CurrentAnimSpaceRootTransform.GetTranslation());
+		TargetAnimSpaceRootTransform.SetTranslation(InverseMeshToActorRotation.GetValue() * TargetAnimSpaceRootTransform.GetTranslation());
+	}
+
+	FTransform TargetTransformWorldSpace = StitchData.TargetTransform * SequenceTransformOrigin;
+
+	FTransform WorldSpaceProjectedFromStart = CurrentAnimSpaceRootTransform.GetRelativeTransform(InitialAnimSpaceRootTransform) * InitialRootTransform;
+	FTransform WorldSpaceProjectedFromTarget = CurrentAnimSpaceRootTransform.GetRelativeTransform(TargetAnimSpaceRootTransform) * TargetTransformWorldSpace;
+
+	FTransform RootTransformWorldSpace;
+	// TODO: Protect against division by 0
+	float BlendWeight = ((CurrentTime - InitialTime) / (TimeToTarget + CurrentTime - InitialTime));
+	RootTransformWorldSpace.Blend(WorldSpaceProjectedFromStart, WorldSpaceProjectedFromTarget, BlendWeight);
+
+	FDeltaTimeRecord DeltaTime;
+	DeltaTime.Set(PreviousAnimationTime, AnimationTime-PreviousAnimationTime);
+	const FAnimExtractContext ExtractionContext((double)AnimationTime, false, DeltaTime, false);
+		
+	if (EnumHasAnyFlags(VM.GetFlags(), EEvaluationFlags::Bones))
+	{
+		FDecompressionTools::GetAnimationPose(AnimSequence, ExtractionContext, Keyframe.Pose, bUseRawData);
+
+		int32 RootIndex = Keyframe.Pose.GetRefPose().GetLODBoneIndexFromSkeletonBoneIndex(0);
+		if (StitchData.SwapRootBone == ESwapRootBone::SwapRootBone_None)
+		{
+			Keyframe.Pose.LocalTransformsView[RootIndex] = RootTransformWorldSpace;
+		}
+		else
+		{
+			Keyframe.Pose.LocalTransformsView[RootIndex] = FTransform::Identity;
+		}
+	}
+
+	if (EnumHasAnyFlags(VM.GetFlags(), EEvaluationFlags::Curves))
+	{
+		FDecompressionTools::GetAnimationCurves(AnimSequence, ExtractionContext, Keyframe.Curves, bUseRawData);
+	}
+
+	if (EnumHasAnyFlags(VM.GetFlags(), EEvaluationFlags::Attributes))
+	{
+		FDecompressionTools::GetAnimationAttributes(AnimSequence, ExtractionContext, Keyframe.Pose.GetRefPose(), Keyframe.Attributes, bUseRawData);
+	}
+
+	if (EnumHasAnyFlags(VM.GetFlags(), EEvaluationFlags::Attributes | EEvaluationFlags::Trajectory))
+	{
+		FTransformAnimationAttribute* RootMotionAttribute = Keyframe.Attributes.FindOrAdd<FTransformAnimationAttribute>(AnimMixerComponents->RootTransformAttributeId);
+		RootMotionAttribute->Value = RootTransformWorldSpace;
+	}
+
+	VM.PushValue(KEYFRAME_STACK_NAME, MakeUnique<FKeyframeState>(MoveTemp(Keyframe)));
+}
+
+UMovieSceneStitchAnimSection::UMovieSceneStitchAnimSection(const FObjectInitializer& ObjectInitializer)
+	: Super(ObjectInitializer)
+{
+	BlendType = EMovieSceneBlendType::Absolute;
+}
+
+EMovieSceneChannelProxyType UMovieSceneStitchAnimSection::CacheChannelProxy()
+{
+	FMovieSceneChannelProxyData Channels;
+
+#if WITH_EDITOR
+
+	static FMovieSceneChannelMetaData MetaData("Weight", LOCTEXT("WeightChannelName", "Weight"));
+	MetaData.bCanCollapseToTrack = false;
+
+	Channels.Add(Weight, MetaData, TMovieSceneExternalValue<float>());
+
+#else
+
+	Channels.Add(Weight);
+
+#endif
+
+	ChannelProxy = MakeShared<FMovieSceneChannelProxy>(MoveTemp(Channels));
+	return EMovieSceneChannelProxyType::Dynamic;
+}
+
+float UMovieSceneStitchAnimSection::GetTotalWeightValue(FFrameTime InTime) const
+{
+	float ManualWeight = 1.f;
+	Weight.Evaluate(InTime, ManualWeight);
+	return ManualWeight * EvaluateEasing(InTime);
+}
+
+UObject* UMovieSceneStitchAnimSection::GetSourceObject() const
+{
+	return StitchDatabase;
+}
+
+void UMovieSceneStitchAnimSection::ImportEntityImpl(UMovieSceneEntitySystemLinker* EntityLinker, const FEntityImportParams& InParams, FImportedEntity* OutImportedEntity)
+{
+	using namespace UE::MovieScene;
+
+	const FBuiltInComponentTypes* BuiltInComponents = FBuiltInComponentTypes::Get();
+	const FMovieSceneTracksComponentTypes* TrackComponents = FMovieSceneTracksComponentTypes::Get();
+	const FPoseSearchTracksComponentTypes* PoseSearchTrackComponents = FPoseSearchTracksComponentTypes::Get();
+	const FAnimMixerComponentTypes* AnimMixerComponents = FAnimMixerComponentTypes::Get();
+
+	const FGuid ObjectBindingID = InParams.GetObjectBindingID();
+
+	if (!ObjectBindingID.IsValid())
+	{
+		return;
+	}
+
+	FMovieSceneStitchAnimComponentData StitchData;
+	StitchData.StitchDatabase = StitchDatabase;
+	StitchData.TargetPoseAsset = TargetPoseAsset;
+	StitchData.TargetAnimationTimeSeconds = TargetAnimationTimeSeconds;
+	StitchData.TargetTransform = TargetTransform;
+	StitchData.StartFrame = GetInclusiveStartFrame();
+	StitchData.EndFrame = GetExclusiveEndFrame();
+	StitchData.SwapRootBone = SwapRootBone;
+	
+	// Make a default eval task. This will be filled out more by the stitch system
+	TSharedPtr<FMovieSceneStitchAnimEvaluationTask> EvalTask = MakeShared<FMovieSceneStitchAnimEvaluationTask>();
+	EvalTask->StitchData = StitchData;
+
+	OutImportedEntity->AddBuilder(
+		FEntityBuilder()
+		.Add(PoseSearchTrackComponents->StitchAnim, StitchData)
+		.Add(BuiltInComponents->GenericObjectBinding, ObjectBindingID)
+		.Add(BuiltInComponents->BoundObjectResolver, UMovieSceneSkeletalAnimationSystem::ResolveSkeletalMeshComponentBinding)
+		.AddConditional(BuiltInComponents->WeightChannel, &Weight, Weight.HasAnyData())
+		.AddTag(TrackComponents->Tags.AnimMixerPoseProducer)
+		.Add(AnimMixerComponents->Priority, MixedAnimationPriority)
+		.Add(AnimMixerComponents->Target, MixedAnimationTarget)
+		.Add(AnimMixerComponents->Task, EvalTask)
+		.Add(AnimMixerComponents->RootMotionSettings, FMovieSceneRootMotionSettings{StitchData.SwapRootBone, EMovieSceneRootMotionSpace::WorldSpace})
+		.AddTag(AnimMixerComponents->Tags.RequiresBlending)
+	);
+}
+
+double FMovieSceneStitchAnimComponentData::MapTimeToSectionSeconds(FFrameTime InPosition, FFrameRate InFrameRate) const
+{
+	return InFrameRate.AsSeconds(InPosition - StartFrame);
+}
+
+#undef LOCTEXT_NAMESPACE 
+
diff --git a/Engine/Plugins/Experimental/MovieScenePoseSearchTracks/Source/MovieScenePoseSearchTracks/Private/Systems/MovieSceneStitchAnimSystem.cpp b/Engine/Plugins/Experimental/MovieScenePoseSearchTracks/Source/MovieScenePoseSearchTracks/Private/Systems/MovieSceneStitchAnimSystem.cpp
new file mode 100644
--- /dev/null
+++ b/Engine/Plugins/Experimental/MovieScenePoseSearchTracks/Source/MovieScenePoseSearchTracks/Private/Systems/MovieSceneStitchAnimSystem.cpp
@@ -0,0 +1,269 @@
+// Copyright Epic Games, Inc. All Rights Reserved.
+
+#include "Systems/MovieSceneStitchAnimSystem.h"
+
+#include "Async/TaskGraphInterfaces.h"
+#include "EntitySystem/BuiltInComponentTypes.h"
+#include "AnimMixerComponentTypes.h"
+#include "PoseSearchTracksComponentTypes.h"
+#include "EntitySystem/MovieSceneEntityIDs.h"
+#include "EntitySystem/MovieSceneEntitySystemLinker.h"
+#include "EntitySystem/MovieSceneEntitySystemRunner.h"
+#include "EntitySystem/MovieSceneEntitySystemTask.h"
+#include "EntitySystem/MovieSceneEvalTimeSystem.h"
+#include "EntitySystem/MovieSceneBoundObjectInstantiator.h"
+#include "EntitySystem/MovieSceneBoundSceneComponentInstantiator.h"
+#include "EntitySystem/MovieScenePreAnimatedStateSystem.h"
+#include "EntitySystem/Interrogation/MovieSceneInterrogationLinker.h"
+#include "MovieSceneTracksComponentTypes.h"
+#include "Sections/MovieSceneStitchAnimSection.h"
+#include "EvaluationVM/EvaluationVM.h"
+#include "EvaluationVM/KeyframeState.h"
+#include "Tracks/IMovieSceneTransformOrigin.h"
+#include "Systems/MovieSceneTransformOriginSystem.h"
+#include "Components/SkeletalMeshComponent.h"
+#include "IMovieScenePlaybackClient.h"
+#include "Animation/AnimInstance.h"
+#include "Component/AnimNextComponent.h"
+#include "Systems/MovieSceneAnimNextTargetSystem.h"
+
+#include UE_INLINE_GENERATED_CPP_BY_NAME(MovieSceneStitchAnimSystem)
+
+
+namespace UE::MovieScene
+{
+	/** ------------------------------------------------------------------------- */
+	/** Updates the stitch anim evaluation task each frame*/
+	struct FUpdateStitchAnimTask
+	{
+		const UMovieSceneEntitySystemLinker* EntityLinker;
+
+		FUpdateStitchAnimTask(const UMovieSceneEntitySystemLinker* InLinker)
+			: EntityLinker(InLinker)
+		{
+		}
+
+		FTransform GetTransformOrigin(FInstanceHandle InstanceHandle) const
+		{
+			FTransform TransformOrigin = FTransform::Identity;
+
+			const FSequenceInstance& Instance = EntityLinker->GetInstanceRegistry()->GetInstance(InstanceHandle);
+			TSharedRef<FSharedPlaybackState> SharedPlaybackState = Instance.GetSharedPlaybackState();
+			const IMovieScenePlaybackClient* Client = SharedPlaybackState->FindCapability<IMovieScenePlaybackClient>();
+			FMovieSceneSequenceID SequenceID = Instance.GetSequenceID();
+			const UObject* InstanceData = Client ? Client->GetInstanceData() : nullptr;
+			const IMovieSceneTransformOrigin* RawInterface = Cast<const IMovieSceneTransformOrigin>(InstanceData);
+
+			const bool bHasInterface = RawInterface || (InstanceData && InstanceData->GetClass()->ImplementsInterface(UMovieSceneTransformOrigin::StaticClass()));
+			if (bHasInterface)
+			{
+				// Retrieve the current origin
+				TransformOrigin = RawInterface ? RawInterface->GetTransformOrigin() : IMovieSceneTransformOrigin::Execute_BP_GetTransformOrigin(InstanceData);
+			}
+
+			if (!EntityLinker)
+			{
+				return TransformOrigin;
+			}
+
+			const UMovieSceneTransformOriginSystem* TransformOriginSystem = EntityLinker->FindSystem<UMovieSceneTransformOriginSystem>();
+
+			if (!TransformOriginSystem)
+			{
+				return TransformOrigin;
+			}
+
+			const TSparseArray<FTransform>& TransformOrigins = TransformOriginSystem->GetTransformOriginsByInstanceID();
+			const TMap<FMovieSceneSequenceID, UE::MovieScene::FInstanceHandle> SequenceIDToInstanceHandle = TransformOriginSystem->GetSequenceIDToInstanceHandle();
+
+			if (SequenceIDToInstanceHandle.Contains(SequenceID))
+			{
+				const UE::MovieScene::FInstanceHandle CurrentHandle = SequenceIDToInstanceHandle[SequenceID];
+				if (TransformOrigins.IsValidIndex(CurrentHandle.InstanceID))
+				{
+					TransformOrigin = TransformOrigins[CurrentHandle.InstanceID];
+				}
+			}
+
+			return TransformOrigin;
+		}
+
+		// Chooses a context object for motion matching calls based on the anim target for the track
+		// TODO: Ideally this would be more encapsulated in the target itself, or the motion matching call wouldn't require this
+		const UObject* GetContextObject(TInstancedStruct<FMovieSceneMixedAnimationTarget> Target, const UObject* BoundObject) const
+		{
+			const UObject* ContextObject = nullptr;
+
+			// If we're targeting anim next, try to grab the anim next component
+			auto TryFindComponent = []<typename T>(const UObject* BoundObject) -> const T*
+				{
+					if (const AActor* Actor = Cast<AActor>(BoundObject))
+					{
+						return Actor->FindComponentByClass<T>();
+					}
+					else if (const UActorComponent* ActorComponent = Cast<UActorComponent>(BoundObject))
+					{
+						if (ActorComponent->GetClass() == T::StaticClass())
+						{
+							return Cast<T>(ActorComponent);
+						}
+						else 
+						{
+							return ActorComponent->GetOwner() ? ActorComponent->GetOwner()->FindComponentByClass<T>() : nullptr;
+						}
+					}
+					return nullptr;
+				};
+
+			if (Target.IsValid() && Target.GetScriptStruct() == FMovieSceneAnimNextInjectionTarget::StaticStruct())
+			{
+				ContextObject = TryFindComponent.template operator()<UAnimNextComponent>(BoundObject);
+			}
+			else if (!Target.IsValid() || Target.GetScriptStruct() == FMovieSceneMixedAnimationTarget::StaticStruct())
+			{
+				// If we specifically didn't choose a target, try seeing if we should use anim next
+				const UAnimNextComponent* AnimNextComponent = TryFindComponent.template operator()<UAnimNextComponent>(BoundObject);
+				const USkeletalMeshComponent* SkeletalMeshComponent = TryFindComponent.template operator()<USkeletalMeshComponent>(BoundObject);
+				if (AnimNextComponent && SkeletalMeshComponent && !SkeletalMeshComponent->bEnableAnimation)
+				{
+					ContextObject = AnimNextComponent;
+				}
+			}
+
+			// Fall back to using the skeletal mesh component
+			if (!ContextObject)
+			{
+				ContextObject = TryFindComponent.template operator()<USkeletalMeshComponent>(BoundObject);
+			}
+
+			return ContextObject;
+		}
+
+		void ForEachAllocation(
+			const FEntityAllocationProxy AllocationProxy,
+			TRead<FInstanceHandle> InstanceHandles,
+			TRead<FMovieSceneStitchAnimComponentData> StitchAnims,
+			TRead<UObject*> BoundObjects,
+			TRead<TInstancedStruct<FMovieSceneMixedAnimationTarget>> Targets,
+			TReadOptional<FFrameTime> OptionalEvalTimes,
+			TWrite<TSharedPtr<FAnimNextEvaluationTask>> EvalTask) const
+		{
+			const FEntityAllocation* Allocation = AllocationProxy.GetAllocation();
+			const int32 Num = Allocation->Num();
+			for (int32 Index = 0; Index < Num; ++Index)
+			{
+				const FInstanceHandle& InstanceHandle(InstanceHandles[Index]);
+				const FMovieSceneStitchAnimComponentData& StitchAnim(StitchAnims[Index]);
+
+				// Get the full context, so we can get both the current and previous evaluation times.
+				const FSequenceInstance& SequenceInstance = EntityLinker->GetInstanceRegistry()->GetInstance(InstanceHandle);
+				const FMovieSceneContext& Context = SequenceInstance.GetContext();
+				const FFrameTime EvalFrameTime = OptionalEvalTimes ? OptionalEvalTimes[Index] : Context.GetTime();
+				const UObject* BoundObject = BoundObjects[Index];
+
+				float PreviousTimeSeconds = StitchAnim.MapTimeToSectionSeconds(Context.GetPreviousTime(), Context.GetFrameRate());
+				float CurrentTimeSeconds = StitchAnim.MapTimeToSectionSeconds(EvalFrameTime, Context.GetFrameRate());
+
+				TSharedPtr<FMovieSceneStitchAnimEvaluationTask> AnimTask = StaticCastSharedPtr<FMovieSceneStitchAnimEvaluationTask>(EvalTask[Index]);
+
+				if (!BoundObject || !StitchAnim.StitchDatabase || !StitchAnim.TargetPoseAsset || !AnimTask.IsValid())
+				{
+					continue;
+				}
+
+				AnimTask->TimeToTarget = StitchAnim.MapTimeToSectionSeconds(StitchAnim.EndFrame, Context.GetFrameRate()) - CurrentTimeSeconds;
+				
+				if (AnimTask->InitialTime < 0.0f)
+				{
+					// We use the previous time for the initial time, as it will match the transform from the actor which hasn't
+					// yet been set for this frame.
+					AnimTask->InitialTime = PreviousTimeSeconds;
+					AnimTask->SequenceTransformOrigin = GetTransformOrigin(InstanceHandle);
+
+					FTransform InitialComponentTransform = FTransform::Identity;
+
+					const AActor* Actor = Cast<AActor>(BoundObject);
+					if (!Actor)
+					{
+						if (const UActorComponent* ActorComponent = Cast<UActorComponent>(BoundObject))
+						{
+							Actor = ActorComponent->GetOwner();
+						}
+					}
+					
+					const USkeletalMeshComponent* SkeletalMeshComponent = Cast<USkeletalMeshComponent>(BoundObject);
+					if (!SkeletalMeshComponent && Actor)
+					{
+						SkeletalMeshComponent = Actor->GetComponentByClass<USkeletalMeshComponent>();
+					}
+
+					AnimTask->ContextObject = GetContextObject(Targets[Index], BoundObject);
+
+					if (Actor && StitchAnim.SwapRootBone == ESwapRootBone::SwapRootBone_Actor)
+					{
+						InitialComponentTransform = Actor->GetRootComponent()->GetRelativeTransform();
+
+						if (SkeletalMeshComponent)
+						{
+							AnimTask->InverseMeshToActorRotation = Actor->GetRootComponent()->GetComponentTransform().GetRelativeTransformReverse(SkeletalMeshComponent->GetComponentTransform()).GetRotation();
+						}
+					}
+					else if (SkeletalMeshComponent && StitchAnim.SwapRootBone == ESwapRootBone::SwapRootBone_Component)
+					{
+						InitialComponentTransform = SkeletalMeshComponent->GetRelativeTransform();
+					}
+
+					AnimTask->InitialRootTransform = InitialComponentTransform;
+				}
+				AnimTask->PreviousTime = PreviousTimeSeconds;
+				AnimTask->CurrentTime = CurrentTimeSeconds;
+			}
+		}
+	};
+
+} // namespace UE::MovieScene
+
+UMovieSceneStitchAnimSystem::UMovieSceneStitchAnimSystem(const FObjectInitializer& ObjInit)
+	: Super(ObjInit)
+{
+	using namespace UE::MovieScene;
+
+	FPoseSearchTracksComponentTypes* PoseSearchTrackComponents = FPoseSearchTracksComponentTypes::Get();
+	FAnimMixerComponentTypes* AnimMixerComponents = FAnimMixerComponentTypes::Get();
+	RelevantComponent = PoseSearchTrackComponents->StitchAnim;
+	Phase = ESystemPhase::Scheduling;
+
+	if (HasAnyFlags(RF_ClassDefaultObject))
+	{
+		FBuiltInComponentTypes* BuiltInComponents = FBuiltInComponentTypes::Get();
+
+		DefineImplicitPrerequisite(UMovieSceneEvalTimeSystem::StaticClass(), GetClass());
+		DefineImplicitPrerequisite(UMovieSceneGenericBoundObjectInstantiator::StaticClass(), GetClass());
+		DefineImplicitPrerequisite(UMovieSceneBoundSceneComponentInstantiator::StaticClass(), GetClass());
+		DefineComponentConsumer(GetClass(), PoseSearchTrackComponents->StitchAnim);
+		DefineImplicitPrerequisite(GetClass(), UMovieSceneRestorePreAnimatedStateSystem::StaticClass());
+		DefineComponentProducer(GetClass(), AnimMixerComponents->Task);
+
+	}
+	SystemCategories |= FSystemInterrogator::GetExcludedFromInterrogationCategory();
+}
+
+void UMovieSceneStitchAnimSystem::OnSchedulePersistentTasks(UE::MovieScene::IEntitySystemScheduler* TaskScheduler)
+{
+	using namespace UE::MovieScene;
+
+	FBuiltInComponentTypes* BuiltInComponents = FBuiltInComponentTypes::Get();
+	FPoseSearchTracksComponentTypes* PoseSearchTrackComponents = FPoseSearchTracksComponentTypes::Get();
+	FAnimMixerComponentTypes* AnimMixerComponents = FAnimMixerComponentTypes::Get();
+
+	FTaskID UpdateTask = FEntityTaskBuilder()
+		.Read(BuiltInComponents->InstanceHandle)
+		.Read(PoseSearchTrackComponents->StitchAnim)
+		.Read(BuiltInComponents->BoundObject)
+		.Read(AnimMixerComponents->Target)
+		.ReadOptional(BuiltInComponents->EvalTime)
+		.Write(AnimMixerComponents->Task)
+		.FilterNone({ BuiltInComponents->Tags.Ignored })
+		.Schedule_PerAllocation<FUpdateStitchAnimTask>(&Linker->EntityManager, TaskScheduler,
+			Linker);
+}
diff --git a/Engine/Plugins/Experimental/MovieScenePoseSearchTracks/Source/MovieScenePoseSearchTracks/Private/Tracks/MovieSceneStitchAnimTrack.cpp b/Engine/Plugins/Experimental/MovieScenePoseSearchTracks/Source/MovieScenePoseSearchTracks/Private/Tracks/MovieSceneStitchAnimTrack.cpp
new file mode 100644
--- /dev/null
+++ b/Engine/Plugins/Experimental/MovieScenePoseSearchTracks/Source/MovieScenePoseSearchTracks/Private/Tracks/MovieSceneStitchAnimTrack.cpp
@@ -0,0 +1,97 @@
+// Copyright Epic Games, Inc. All Rights Reserved.
+
+#include "Tracks/MovieSceneStitchAnimTrack.h"
+#include "Sections/MovieSceneStitchAnimSection.h"
+#include "PoseSearch/PoseSearchDatabase.h"
+#include "MovieScene.h"
+
+#include UE_INLINE_GENERATED_CPP_BY_NAME(MovieSceneStitchAnimTrack)
+
+#define LOCTEXT_NAMESPACE "MovieSceneStitchAnimTrack"
+
+UMovieSceneStitchAnimTrack::UMovieSceneStitchAnimTrack(const FObjectInitializer& ObjectInitializer)
+	: Super(ObjectInitializer)
+{
+#if WITH_EDITORONLY_DATA
+	TrackTint = FColor(124, 15, 124, 65);
+	bSupportsDefaultSections = false;
+#endif
+
+	SupportedBlendTypes.Add(EMovieSceneBlendType::Absolute);
+
+	EvalOptions.bCanEvaluateNearestSection = true;
+}
+
+const TArray<UMovieSceneSection*>& UMovieSceneStitchAnimTrack::GetAllSections() const
+{
+	return AnimationSections;
+}
+
+
+bool UMovieSceneStitchAnimTrack::SupportsMultipleRows() const
+{
+	return true;
+}
+
+bool UMovieSceneStitchAnimTrack::SupportsType(TSubclassOf<UMovieSceneSection> SectionClass) const
+{
+	return SectionClass == UMovieSceneStitchAnimTrack::StaticClass();
+}
+
+UMovieSceneSection* UMovieSceneStitchAnimTrack::CreateNewSection()
+{
+	return NewObject<UMovieSceneStitchAnimSection>(this, NAME_None, RF_Transactional);
+}
+
+bool UMovieSceneStitchAnimTrack::HasSection(const UMovieSceneSection& Section) const
+{
+	return AnimationSections.Contains(&Section);
+}
+
+void UMovieSceneStitchAnimTrack::AddSection(UMovieSceneSection& Section)
+{
+	AnimationSections.Add(&Section);
+}
+
+void UMovieSceneStitchAnimTrack::RemoveSection(UMovieSceneSection& Section)
+{
+	AnimationSections.Remove(&Section);
+}
+
+void UMovieSceneStitchAnimTrack::RemoveSectionAt(int32 SectionIndex)
+{
+	AnimationSections.RemoveAt(SectionIndex);
+}
+
+bool UMovieSceneStitchAnimTrack::IsEmpty() const
+{
+	return AnimationSections.Num() == 0;
+}
+
+#if WITH_EDITORONLY_DATA
+
+FText UMovieSceneStitchAnimTrack::GetDefaultDisplayName() const
+{
+	return LOCTEXT("TrackName", "Stitch");
+}
+
+#endif
+
+UMovieSceneSection* UMovieSceneStitchAnimTrack::AddNewAnimationOnRow(FFrameNumber KeyTime, UPoseSearchDatabase* PoseSearchDatabase, int32 RowIndex)
+{
+	UMovieSceneStitchAnimSection* NewSection = Cast<UMovieSceneStitchAnimSection>(CreateNewSection());
+	{
+		FFrameTime AnimationLength = 5.0f * GetTypedOuter<UMovieScene>()->GetTickResolution();
+		int32 IFrameNumber = AnimationLength.FrameNumber.Value + (int)(AnimationLength.GetSubFrame() + 0.5f) + 1;
+
+		NewSection->InitialPlacementOnRow(AnimationSections, KeyTime, IFrameNumber, RowIndex);
+		NewSection->StitchDatabase = PoseSearchDatabase;
+	}
+
+	AddSection(*NewSection);
+
+	return NewSection;
+}
+
+#undef LOCTEXT_NAMESPACE
+
diff --git a/Engine/Plugins/Experimental/MovieScenePoseSearchTracks/Source/MovieScenePoseSearchTracks/Public/PoseSearchTracksComponentTypes.h b/Engine/Plugins/Experimental/MovieScenePoseSearchTracks/Source/MovieScenePoseSearchTracks/Public/PoseSearchTracksComponentTypes.h
new file mode 100644
--- /dev/null
+++ b/Engine/Plugins/Experimental/MovieScenePoseSearchTracks/Source/MovieScenePoseSearchTracks/Public/PoseSearchTracksComponentTypes.h
@@ -0,0 +1,21 @@
+// Copyright Epic Games, Inc. All Rights Reserved.
+
+#pragma once
+
+#include "EntitySystem/MovieSceneEntityIDs.h"
+#include "Sections/MovieSceneStitchAnimSection.h"
+
+namespace UE::MovieScene
+{
+	struct FPoseSearchTracksComponentTypes
+	{
+		MOVIESCENEPOSESEARCHTRACKS_API static FPoseSearchTracksComponentTypes* Get();
+		MOVIESCENEPOSESEARCHTRACKS_API static void Destroy();
+
+		TComponentTypeID<FMovieSceneStitchAnimComponentData> StitchAnim;
+
+	private:
+		FPoseSearchTracksComponentTypes();
+	};
+
+}
\ No newline at end of file
diff --git a/Engine/Plugins/Experimental/MovieScenePoseSearchTracks/Source/MovieScenePoseSearchTracks/Public/Sections/MovieSceneStitchAnimSection.h b/Engine/Plugins/Experimental/MovieScenePoseSearchTracks/Source/MovieScenePoseSearchTracks/Public/Sections/MovieSceneStitchAnimSection.h
new file mode 100644
--- /dev/null
+++ b/Engine/Plugins/Experimental/MovieScenePoseSearchTracks/Source/MovieScenePoseSearchTracks/Public/Sections/MovieSceneStitchAnimSection.h
@@ -0,0 +1,157 @@
+// Copyright Epic Games, Inc. All Rights Reserved.
+
+#pragma once
+
+#include "CoreMinimal.h"
+#include "UObject/ObjectMacros.h"
+#include "MovieSceneSection.h"
+#include "Channels/MovieSceneFloatChannel.h"
+#include "EntitySystem/IMovieSceneEntityProvider.h"
+#include "StructUtils/InstancedStruct.h"
+#include "MovieSceneMixedAnimationTarget.h"
+#include "Animation/AnimationAsset.h"
+#include "EvaluationVM/EvaluationTask.h"
+#include "AnimSequencerInstanceProxy.h"
+#include "MovieSceneStitchAnimSection.generated.h"
+
+USTRUCT()
+struct FMovieSceneStitchAnimComponentData
+{
+	GENERATED_BODY()
+
+	UPROPERTY()
+	TObjectPtr<UObject> StitchDatabase;
+
+	UPROPERTY()
+	TObjectPtr<UAnimationAsset> TargetPoseAsset;
+	
+	UPROPERTY()
+	float TargetAnimationTimeSeconds = 0.0f;
+
+	UPROPERTY()
+	FTransform TargetTransform;
+
+	UPROPERTY()
+	FFrameNumber StartFrame;
+
+	UPROPERTY()
+	FFrameNumber EndFrame;
+
+	UPROPERTY()
+	ESwapRootBone SwapRootBone = ESwapRootBone::SwapRootBone_None;
+
+	double MapTimeToSectionSeconds(FFrameTime InPosition, FFrameRate InFrameRate) const;
+};
+
+USTRUCT()
+struct FMovieSceneStitchAnimEvaluationTask : public FAnimNextEvaluationTask
+{
+	GENERATED_BODY()
+
+	DECLARE_ANIM_EVALUATION_TASK(FMovieSceneStitchAnimEvaluationTask)
+
+	// Task entry point
+	virtual void Execute(UE::AnimNext::FEvaluationVM& VM) const override;
+
+	UPROPERTY()
+	FMovieSceneStitchAnimComponentData StitchData;
+
+	// Initial sequencer time when the motion matching calculation was done
+	UPROPERTY()
+	float InitialTime = -1.0f;
+
+	// Time until our target pose
+	UPROPERTY()
+	float TimeToTarget = -1.0f;
+
+	// Previous sequencer time
+	UPROPERTY()
+	float PreviousTime = -1.0f;
+
+	// Current sequencer time
+	UPROPERTY()
+	float CurrentTime = -1.0f;
+
+	UPROPERTY()
+	FTransform InitialRootTransform;
+
+	UPROPERTY()
+	TWeakObjectPtr<const UObject> ContextObject;
+
+	UPROPERTY()
+	FTransform SequenceTransformOrigin;
+
+	// Calculated and cached by the motion match at InitialTime
+	UPROPERTY()
+	mutable TObjectPtr<UAnimationAsset> MatchedAsset;
+
+	UPROPERTY()
+	mutable float MatchedAssetTime;
+
+	UPROPERTY()
+	mutable float MatchedAssetPlayRate = 1.0f;
+
+	UPROPERTY()
+	TOptional<FQuat> InverseMeshToActorRotation;
+};
+
+
+/**
+ * Handles generating and playing back transitional skeletal animations from a stitch database.
+ */
+UCLASS( MinimalAPI )
+class UMovieSceneStitchAnimSection
+	: public UMovieSceneSection
+	, public IMovieSceneEntityProvider
+{
+	GENERATED_BODY()
+
+public:
+
+	UMovieSceneStitchAnimSection(const FObjectInitializer& ObjectInitializer);
+
+	UPROPERTY(EditAnywhere, Category = "Stitch", meta=(AllowedClasses = "/Script/PoseSearch.PoseSearchDatabase"))
+	TObjectPtr<UObject> StitchDatabase;
+
+	UPROPERTY(EditAnywhere, Category = "Stitch")
+	TObjectPtr<UAnimationAsset> TargetPoseAsset;
+
+	// TODO: Is this the best way to specify this?
+	UPROPERTY(EditAnywhere, Category = "Stitch")
+	float TargetAnimationTimeSeconds = 0.0f;
+
+	UPROPERTY(EditAnywhere, Category = "Stitch")
+	FTransform TargetTransform;
+
+	// TODO: Would be nice to standardize these properties across multiple anim tracks
+	UPROPERTY(EditAnywhere, Category = "Root Motion")
+	bool bBlendFirstChildOfRoot = false;
+	
+	UPROPERTY(EditAnywhere, Category = "Root Motion")
+	int32 ChildBoneIndex = INDEX_NONE;
+
+	UPROPERTY(EditAnywhere, Category = "Root Motion")
+	ESwapRootBone SwapRootBone = ESwapRootBone::SwapRootBone_Actor;
+
+	UPROPERTY(EditAnywhere, Category = "Mixed Animation")
+	TInstancedStruct<FMovieSceneMixedAnimationTarget> MixedAnimationTarget;
+
+	UPROPERTY(EditAnywhere, Category = "Mixed Animation")
+	int32 MixedAnimationPriority = 0;
+
+	/** The weight curve for this animation section */
+	UPROPERTY()
+	FMovieSceneFloatChannel Weight;
+
+	//~ UMovieSceneSection interface
+	virtual EMovieSceneChannelProxyType CacheChannelProxy() override;
+	virtual UObject* GetSourceObject() const override;
+
+protected:
+
+	virtual float GetTotalWeightValue(FFrameTime InTime) const override;
+
+	/** ~IMovieSceneEntityProvider interface */
+	virtual void ImportEntityImpl(UMovieSceneEntitySystemLinker* EntityLinker, const FEntityImportParams& InParams, FImportedEntity* OutImportedEntity) override;
+
+};
diff --git a/Engine/Plugins/Experimental/MovieScenePoseSearchTracks/Source/MovieScenePoseSearchTracks/Public/Systems/MovieSceneStitchAnimSystem.h b/Engine/Plugins/Experimental/MovieScenePoseSearchTracks/Source/MovieScenePoseSearchTracks/Public/Systems/MovieSceneStitchAnimSystem.h
new file mode 100644
--- /dev/null
+++ b/Engine/Plugins/Experimental/MovieScenePoseSearchTracks/Source/MovieScenePoseSearchTracks/Public/Systems/MovieSceneStitchAnimSystem.h
@@ -0,0 +1,27 @@
+
+// Copyright Epic Games, Inc. All Rights Reserved.
+
+#pragma once
+
+#include "EntitySystem/MovieSceneEntitySystem.h"
+#include "EntitySystem/MovieSceneEntityInstantiatorSystem.h"
+#include "EvaluationVM/EvaluationTask.h"
+#include "MovieSceneStitchAnimSystem.generated.h"
+
+// System to handle updating stitch anim evaluation tasks for the anim mixer.
+UCLASS(MinimalAPI)
+class UMovieSceneStitchAnimSystem
+	: public UMovieSceneEntityInstantiatorSystem
+{
+public:
+
+	GENERATED_BODY()
+
+	UMovieSceneStitchAnimSystem(const FObjectInitializer& ObjInit);
+
+private:
+
+	virtual void OnSchedulePersistentTasks(UE::MovieScene::IEntitySystemScheduler* TaskScheduler) override;
+	
+
+};
\ No newline at end of file
diff --git a/Engine/Plugins/Experimental/MovieScenePoseSearchTracks/Source/MovieScenePoseSearchTracks/Public/Tracks/MovieSceneStitchAnimTrack.h b/Engine/Plugins/Experimental/MovieScenePoseSearchTracks/Source/MovieScenePoseSearchTracks/Public/Tracks/MovieSceneStitchAnimTrack.h
new file mode 100644
--- /dev/null
+++ b/Engine/Plugins/Experimental/MovieScenePoseSearchTracks/Source/MovieScenePoseSearchTracks/Public/Tracks/MovieSceneStitchAnimTrack.h
@@ -0,0 +1,59 @@
+// Copyright Epic Games, Inc. All Rights Reserved.
+
+#pragma once
+
+#include "CoreMinimal.h"
+#include "UObject/ObjectMacros.h"
+#include "MovieSceneNameableTrack.h"
+#include "Sections/MovieSceneStitchAnimSection.h"
+#include "MovieSceneStitchAnimTrack.generated.h"
+
+
+/**
+ * Handles generating and playing back transitional skeletal animations from a stitch database.
+ */
+UCLASS(MinimalAPI)
+class UMovieSceneStitchAnimTrack
+	: public UMovieSceneNameableTrack
+{
+	GENERATED_BODY()
+
+public:
+
+	UMovieSceneStitchAnimTrack(const FObjectInitializer& ObjectInitializer);
+
+public:
+
+	// UMovieSceneTrack interface
+
+	virtual bool HasSection(const UMovieSceneSection& Section) const override;
+	virtual void AddSection(UMovieSceneSection& Section) override;
+	virtual void RemoveSection(UMovieSceneSection& Section) override;
+	virtual void RemoveSectionAt(int32 SectionIndex) override;
+	virtual bool IsEmpty() const override;
+	virtual const TArray<UMovieSceneSection*>& GetAllSections() const override;
+	virtual bool SupportsType(TSubclassOf<UMovieSceneSection> SectionClass) const override;
+	virtual UMovieSceneSection* CreateNewSection() override;
+	virtual bool SupportsMultipleRows() const override;
+
+#if WITH_EDITORONLY_DATA
+	virtual FText GetDefaultDisplayName() const override;
+#endif
+
+	/** Adds a new animation to this track */
+	MOVIESCENEPOSESEARCHTRACKS_API UMovieSceneSection* AddNewAnimationOnRow(FFrameNumber KeyTime, class UPoseSearchDatabase* PoseSearchDatabase, int32 RowIndex);
+
+public:
+
+	/** List of all animation sections */
+	UPROPERTY()
+	TArray<TObjectPtr<UMovieSceneSection>> AnimationSections;
+};
+
+
+
+
+
+
+
+
diff --git a/Engine/Plugins/Experimental/MovieScenePoseSearchTracks/Source/MovieScenePoseSearchTracksEditor/MovieScenePoseSearchTracksEditor.Build.cs b/Engine/Plugins/Experimental/MovieScenePoseSearchTracks/Source/MovieScenePoseSearchTracksEditor/MovieScenePoseSearchTracksEditor.Build.cs
new file mode 100644
--- /dev/null
+++ b/Engine/Plugins/Experimental/MovieScenePoseSearchTracks/Source/MovieScenePoseSearchTracksEditor/MovieScenePoseSearchTracksEditor.Build.cs
@@ -0,0 +1,33 @@
+// Copyright Epic Games, Inc. All Rights Reserved.
+
+namespace UnrealBuildTool.Rules
+{
+	public class MovieScenePoseSearchTracksEditor : ModuleRules
+	{
+		public MovieScenePoseSearchTracksEditor(ReadOnlyTargetRules Target) : base(Target)
+		{
+			PublicDependencyModuleNames.AddRange(
+				new string[]
+				{
+					"Core",
+					"CoreUObject",
+					"Engine",
+					"MovieScene",
+					"MovieSceneTracks",
+					"MovieSceneTools",
+					"MovieScenePoseSearchTracks",
+					"AnimGraphRuntime",
+					"AnimNext",
+					"AnimNextAnimGraph",
+					"AnimNextPoseSearch",
+					"Sequencer",
+					"SequencerCore",
+					"SlateCore",
+					"Slate",
+					"UnrealEd",
+					"PoseSearch"
+				}
+			);
+		}
+	}
+}
\ No newline at end of file
diff --git a/Engine/Plugins/Experimental/MovieScenePoseSearchTracks/Source/MovieScenePoseSearchTracksEditor/Private/MovieScenePoseSearchTracksEditorModule.cpp b/Engine/Plugins/Experimental/MovieScenePoseSearchTracks/Source/MovieScenePoseSearchTracksEditor/Private/MovieScenePoseSearchTracksEditorModule.cpp
new file mode 100644
--- /dev/null
+++ b/Engine/Plugins/Experimental/MovieScenePoseSearchTracks/Source/MovieScenePoseSearchTracksEditor/Private/MovieScenePoseSearchTracksEditorModule.cpp
@@ -0,0 +1,40 @@
+// Copyright Epic Games, Inc. All Rights Reserved.
+
+#include "MovieScenePoseSearchTracksEditorModule.h"
+#include "Modules/ModuleInterface.h"
+#include "Modules/ModuleManager.h"
+#include "ISequencerModule.h"
+#include "TrackEditors/StitchAnimTrackEditor.h"
+
+namespace UE::MovieScene
+{
+
+void FMovieScenePoseSearchTracksEditorModule::StartupModule()
+{
+	using namespace UE::Sequencer;
+
+	if (GIsEditor)
+	{
+		ISequencerModule& SequencerModule = FModuleManager::Get().LoadModuleChecked<ISequencerModule>("Sequencer");
+
+		// register specialty track editors
+		StitchAnimationTrackCreateEditorHandle = SequencerModule.RegisterTrackEditor(FOnCreateTrackEditor::CreateStatic(&FStitchAnimTrackEditor::CreateTrackEditor));
+	}
+}
+
+void FMovieScenePoseSearchTracksEditorModule::ShutdownModule()
+{
+	if (!FModuleManager::Get().IsModuleLoaded("Sequencer"))
+	{
+		return;
+	}
+
+	ISequencerModule& SequencerModule = FModuleManager::Get().GetModuleChecked<ISequencerModule>("Sequencer");
+
+	// unregister specialty track editors
+	SequencerModule.UnRegisterTrackEditor(StitchAnimationTrackCreateEditorHandle);
+}
+
+}
+
+IMPLEMENT_MODULE(UE::MovieScene::FMovieScenePoseSearchTracksEditorModule, MovieScenePoseSearchTracksEditor)
diff --git a/Engine/Plugins/Experimental/MovieScenePoseSearchTracks/Source/MovieScenePoseSearchTracksEditor/Private/MovieScenePoseSearchTracksEditorModule.h b/Engine/Plugins/Experimental/MovieScenePoseSearchTracks/Source/MovieScenePoseSearchTracksEditor/Private/MovieScenePoseSearchTracksEditorModule.h
new file mode 100644
--- /dev/null
+++ b/Engine/Plugins/Experimental/MovieScenePoseSearchTracks/Source/MovieScenePoseSearchTracksEditor/Private/MovieScenePoseSearchTracksEditorModule.h
@@ -0,0 +1,20 @@
+﻿// Copyright Epic Games, Inc. All Rights Reserved.
+
+#pragma once
+
+#include "CoreMinimal.h"
+#include "Modules/ModuleInterface.h"
+
+namespace UE::MovieScene
+{
+
+class FMovieScenePoseSearchTracksEditorModule : public IModuleInterface
+{
+	// IModuleInterface interface
+	virtual void StartupModule() override;
+	virtual void ShutdownModule() override;
+
+	FDelegateHandle StitchAnimationTrackCreateEditorHandle;
+};
+
+}
diff --git a/Engine/Plugins/Experimental/MovieScenePoseSearchTracks/Source/MovieScenePoseSearchTracksEditor/Private/TrackEditors/StitchAnimTrackEditor.cpp b/Engine/Plugins/Experimental/MovieScenePoseSearchTracks/Source/MovieScenePoseSearchTracksEditor/Private/TrackEditors/StitchAnimTrackEditor.cpp
new file mode 100644
--- /dev/null
+++ b/Engine/Plugins/Experimental/MovieScenePoseSearchTracks/Source/MovieScenePoseSearchTracksEditor/Private/TrackEditors/StitchAnimTrackEditor.cpp
@@ -0,0 +1,584 @@
+// Copyright Epic Games, Inc. All Rights Reserved.
+
+#include "TrackEditors/StitchAnimTrackEditor.h"
+#include "Sections/MovieSceneStitchAnimSection.h"
+#include "Tracks/MovieSceneStitchAnimTrack.h"
+#include "GameFramework/Actor.h"
+#include "Components/SkeletalMeshComponent.h"
+#include "Engine/SkeletalMesh.h"
+#include "Engine/BlueprintGeneratedClass.h"
+#include "Engine/SCS_Node.h"
+#include "MVVM/ViewModels/ViewDensity.h"
+#include "SequencerSectionPainter.h"
+#include "AssetRegistry/AssetRegistryModule.h"
+#include "PoseSearch/PoseSearchDatabase.h"
+#include "Framework/MultiBox/MultiBoxBuilder.h"
+#include "MVVM/Extensions/ITrackExtension.h"
+#include "IContentBrowserSingleton.h"
+#include "ContentBrowserModule.h"
+#include "SequencerSettings.h"
+#include "MVVM/Views/ViewUtilities.h"
+#include "DragAndDrop/AssetDragDropOp.h"
+#include "MovieSceneToolHelpers.h"
+
+namespace StitchAnimEditorConstants
+{
+	// @todo Sequencer Allow this to be customizable
+	const uint32 AnimationTrackHeight = 28;
+}
+
+#define LOCTEXT_NAMESPACE "FStitchAnimTrackEditor"
+
+USkeletalMeshComponent* AcquireSkeletalMeshFromObjectGuid(const FGuid& Guid, TSharedPtr<ISequencer> SequencerPtr)
+{
+	UObject* BoundObject = SequencerPtr.IsValid() ? SequencerPtr->FindSpawnedObjectOrTemplate(Guid) : nullptr;
+
+	if (AActor* Actor = Cast<AActor>(BoundObject))
+	{
+		if (USkeletalMeshComponent* SkeletalMeshComponent = Cast<USkeletalMeshComponent>(Actor->GetRootComponent()))
+		{
+			return SkeletalMeshComponent;
+		}
+
+		TArray<USkeletalMeshComponent*> SkeletalMeshComponents;
+		Actor->GetComponents(SkeletalMeshComponents);
+
+		if (SkeletalMeshComponents.Num() == 1)
+		{
+			return SkeletalMeshComponents[0];
+		}
+	}
+	else if (USkeletalMeshComponent* SkeletalMeshComponent = Cast<USkeletalMeshComponent>(BoundObject))
+	{
+		if (SkeletalMeshComponent->GetSkeletalMeshAsset())
+		{
+			return SkeletalMeshComponent;
+		}
+	}
+
+	return nullptr;
+}
+
+USkeleton* GetSkeletonFromComponent(UActorComponent* InComponent)
+{
+	USkeletalMeshComponent* SkeletalMeshComp = Cast<USkeletalMeshComponent>(InComponent);
+	if (SkeletalMeshComp && SkeletalMeshComp->GetSkeletalMeshAsset() && SkeletalMeshComp->GetSkeletalMeshAsset()->GetSkeleton())
+	{
+		// @todo Multiple actors, multiple components
+		return SkeletalMeshComp->GetSkeletalMeshAsset()->GetSkeleton();
+	}
+
+	return nullptr;
+}
+
+// Get the skeletal mesh components from the guid
+// If bGetSingleRootComponent - return only the root component if it is a skeletal mesh component. 
+// This allows the root object binding to have an animation track without needing a skeletal mesh component binding
+//
+TArray<USkeletalMeshComponent*> AcquireSkeletalMeshComponentsFromObjectGuid(const FGuid& Guid, TSharedPtr<ISequencer> SequencerPtr, const bool bGetSingleRootComponent = true)
+{
+	TArray<USkeletalMeshComponent*> SkeletalMeshComponents;
+
+	UObject* BoundObject = SequencerPtr.IsValid() ? SequencerPtr->FindSpawnedObjectOrTemplate(Guid) : nullptr;
+
+	AActor* Actor = Cast<AActor>(BoundObject);
+
+	if (!Actor)
+	{
+		if (UChildActorComponent* ChildActorComponent = Cast<UChildActorComponent>(BoundObject))
+		{
+			Actor = ChildActorComponent->GetChildActor();
+		}
+	}
+
+	if (Actor)
+	{
+		if (bGetSingleRootComponent)
+		{
+			if (USkeletalMeshComponent* SkeletalMeshComponent = Cast<USkeletalMeshComponent>(Actor->GetRootComponent()))
+			{
+				SkeletalMeshComponents.Add(SkeletalMeshComponent);
+				return SkeletalMeshComponents;
+			}
+		}
+
+		Actor->GetComponents(SkeletalMeshComponents);
+		if (SkeletalMeshComponents.Num())
+		{
+			return SkeletalMeshComponents;
+		}
+
+		AActor* ActorCDO = Cast<AActor>(Actor->GetClass()->GetDefaultObject());
+		if (ActorCDO)
+		{
+			if (bGetSingleRootComponent)
+			{
+				if (USkeletalMeshComponent* SkeletalMeshComponent = Cast<USkeletalMeshComponent>(ActorCDO->GetRootComponent()))
+				{
+					SkeletalMeshComponents.Add(SkeletalMeshComponent);
+					return SkeletalMeshComponents;
+				}
+			}
+
+			ActorCDO->GetComponents(SkeletalMeshComponents);
+			if (SkeletalMeshComponents.Num())
+			{
+				return SkeletalMeshComponents;
+			}
+		}
+
+		UBlueprintGeneratedClass* ActorBlueprintGeneratedClass = Cast<UBlueprintGeneratedClass>(Actor->GetClass());
+		if (ActorBlueprintGeneratedClass && ActorBlueprintGeneratedClass->SimpleConstructionScript)
+		{
+			const TArray<USCS_Node*>& ActorBlueprintNodes = ActorBlueprintGeneratedClass->SimpleConstructionScript->GetAllNodes();
+
+			for (USCS_Node* Node : ActorBlueprintNodes)
+			{
+				if (Node->ComponentClass->IsChildOf(USkeletalMeshComponent::StaticClass()))
+				{
+					if (USkeletalMeshComponent* SkeletalMeshComponent = Cast<USkeletalMeshComponent>(Node->GetActualComponentTemplate(ActorBlueprintGeneratedClass)))
+					{
+						SkeletalMeshComponents.Add(SkeletalMeshComponent);
+					}
+				}
+			}
+
+			if (SkeletalMeshComponents.Num())
+			{
+				return SkeletalMeshComponents;
+			}
+		}
+	}
+	else if (USkeletalMeshComponent* SkeletalMeshComponent = Cast<USkeletalMeshComponent>(BoundObject))
+	{
+		SkeletalMeshComponents.Add(SkeletalMeshComponent);
+		return SkeletalMeshComponents;
+	}
+
+	return SkeletalMeshComponents;
+}
+
+USkeleton* AcquireSkeletonFromObjectGuid(const FGuid& Guid, TSharedPtr<ISequencer> SequencerPtr)
+{
+	TArray<USkeletalMeshComponent*> SkeletalMeshComponents = AcquireSkeletalMeshComponentsFromObjectGuid(Guid, SequencerPtr);
+
+	if (SkeletalMeshComponents.Num() == 1)
+	{
+		return GetSkeletonFromComponent(SkeletalMeshComponents[0]);
+	}
+
+	return nullptr;
+}
+
+FStitchAnimSection::FStitchAnimSection(UMovieSceneSection& InSection, TWeakPtr<ISequencer> InSequencer)
+	: Section(*CastChecked<UMovieSceneStitchAnimSection>(&InSection))
+	, Sequencer(InSequencer)
+{
+}
+
+FStitchAnimSection::~FStitchAnimSection()
+{
+}
+
+UMovieSceneSection* FStitchAnimSection::GetSectionObject()
+{
+	return &Section;
+}
+
+FText FStitchAnimSection::GetSectionTitle() const
+{
+	if (Section.StitchDatabase != nullptr)
+	{
+		return FText::Format(LOCTEXT("SectionTitleContentFormat", "{0}"), FText::FromString(Section.StitchDatabase->GetName()));
+	}
+	return LOCTEXT("NoStitchSection", "No Stitch Database");
+}
+
+FText FStitchAnimSection::GetSectionToolTip() const
+{
+	if (Section.StitchDatabase != nullptr && Section.TargetPoseAsset != nullptr)
+	{
+		return FText::Format(LOCTEXT("ToolTipContentFormat", "Stitch using database {0} to match pose from {1}"), FText::FromString(Section.StitchDatabase->GetName()), FText::FromString(Section.TargetPoseAsset->GetName()));
+	}
+	return FText::GetEmpty();
+}
+
+float FStitchAnimSection::GetSectionHeight(const UE::Sequencer::FViewDensityInfo& ViewDensity) const
+{
+	return ViewDensity.UniformHeight.Get(StitchAnimEditorConstants::AnimationTrackHeight);
+}
+
+
+FMargin FStitchAnimSection::GetContentPadding() const
+{
+	return FMargin(8.0f, 8.0f);
+}
+
+int32 FStitchAnimSection::OnPaintSection(FSequencerSectionPainter& Painter) const
+{
+	return Painter.PaintSectionBackground();
+}
+
+FStitchAnimTrackEditor::FStitchAnimTrackEditor(TSharedRef<ISequencer> InSequencer)
+	: FMovieSceneTrackEditor(InSequencer)
+{
+
+}
+
+TSharedRef<ISequencerTrackEditor> FStitchAnimTrackEditor::CreateTrackEditor(TSharedRef<ISequencer> InSequencer)
+{
+	return MakeShareable(new FStitchAnimTrackEditor(InSequencer));
+}
+
+bool FStitchAnimTrackEditor::SupportsSequence(UMovieSceneSequence* InSequence) const
+{
+	ETrackSupport TrackSupported = InSequence ? InSequence->IsTrackSupported(UMovieSceneStitchAnimTrack::StaticClass()) : ETrackSupport::NotSupported;
+	return TrackSupported != ETrackSupport::NotSupported;
+}
+
+bool FStitchAnimTrackEditor::SupportsType(TSubclassOf<UMovieSceneTrack> Type) const
+{
+	return Type == UMovieSceneStitchAnimTrack::StaticClass();
+}
+
+TSharedRef<ISequencerSection> FStitchAnimTrackEditor::MakeSectionInterface(UMovieSceneSection& SectionObject, UMovieSceneTrack& Track, FGuid ObjectBinding)
+{
+	check(SupportsType(SectionObject.GetOuter()->GetClass()));
+
+	return MakeShareable(new FStitchAnimSection(SectionObject, GetSequencer()));
+}
+
+bool FStitchAnimTrackEditor::HandleAssetAdded(UObject* Asset, const FGuid& TargetObjectGuid)
+{
+	using namespace UE::PoseSearch;
+	TSharedPtr<ISequencer> SequencerPtr = GetSequencer();
+
+	if (Asset->IsA<UPoseSearchDatabase>() && SequencerPtr.IsValid())
+	{
+		UPoseSearchDatabase* PoseSearchDatabase = Cast<UPoseSearchDatabase>(Asset);
+
+		if (TargetObjectGuid.IsValid() && PoseSearchDatabase)
+		{
+			USkeleton* Skeleton = AcquireSkeletonFromObjectGuid(TargetObjectGuid, GetSequencer());
+
+			// TODO: Check skeleton compatibility? Do I need to specify a role for the database?
+			//if (Skeleton && Skeleton->IsCompatibleForEditor(AnimSequence->GetSkeleton()))
+			{
+				UObject* Object = SequencerPtr->FindSpawnedObjectOrTemplate(TargetObjectGuid);
+
+				UMovieSceneTrack* Track = nullptr;
+
+				const FScopedTransaction Transaction(LOCTEXT("AddStitchedAnim_Transaction", "Add Stitched Animation"));
+
+				int32 RowIndex = INDEX_NONE;
+				AnimatablePropertyChanged(FOnKeyProperty::CreateRaw(this, &FStitchAnimTrackEditor::AddKeyInternal, Object, PoseSearchDatabase, Track, RowIndex));
+
+				return true;
+			}
+		}
+	}
+	return false;
+}
+
+void FStitchAnimTrackEditor::BuildObjectBindingTrackMenu(FMenuBuilder& MenuBuilder, const TArray<FGuid>& ObjectBindings, const UClass* ObjectClass)
+{
+	using namespace UE::PoseSearch;
+	if (ObjectClass->IsChildOf(USkeletalMeshComponent::StaticClass()) || ObjectClass->IsChildOf(AActor::StaticClass()) || ObjectClass->IsChildOf(UChildActorComponent::StaticClass()))
+	{
+		const TSharedPtr<ISequencer> ParentSequencer = GetSequencer();
+
+		USkeleton* Skeleton = AcquireSkeletonFromObjectGuid(ObjectBindings[0], GetSequencer());
+
+		if (Skeleton)
+		{
+			// Load the asset registry module
+			FAssetRegistryModule& AssetRegistryModule = FModuleManager::LoadModuleChecked<FAssetRegistryModule>(TEXT("AssetRegistry"));
+
+			// Collect a full list of assets with the specified class
+			TArray<FAssetData> AssetDataList;
+			AssetRegistryModule.Get().GetAssetsByClass(UPoseSearchDatabase::StaticClass()->GetClassPathName(), AssetDataList, true);
+
+			if (AssetDataList.Num())
+			{
+				UMovieSceneTrack* Track = nullptr;
+
+				MenuBuilder.AddSubMenu(
+					LOCTEXT("AddStitchAnimation", "Stitch Animation"), NSLOCTEXT("Sequencer", "AddStitchAnimationTooltip", "Adds a stitch animation track."),
+					FNewMenuDelegate::CreateRaw(this, &FStitchAnimTrackEditor::AddAnimationSubMenu, ObjectBindings, Skeleton, Track)
+				);
+			}
+		}
+	}
+}
+
+TSharedRef<SWidget> FStitchAnimTrackEditor::BuildAddAnimationSubMenu(FGuid ObjectBinding, USkeleton* Skeleton, UE::Sequencer::TWeakViewModelPtr<UE::Sequencer::ITrackExtension> WeakTrackModel)
+{
+	FMenuBuilder MenuBuilder(true, nullptr);
+
+	TArray<FGuid> ObjectBindings;
+	ObjectBindings.Add(ObjectBinding);
+
+	MenuBuilder.BeginSection(NAME_None, LOCTEXT("AddStitchAnimation_Label", "Add Stitch Animation"));
+	{
+		AddAnimationSubMenu(MenuBuilder, ObjectBindings, Skeleton, WeakTrackModel.Pin()->GetTrack());
+	}
+	MenuBuilder.EndSection();
+
+	return MenuBuilder.MakeWidget();
+}
+
+TSharedRef<SWidget> FStitchAnimTrackEditor::BuildAnimationSubMenu(FGuid ObjectBinding, USkeleton* Skeleton, UMovieSceneTrack* Track)
+{
+	FMenuBuilder MenuBuilder(true, nullptr);
+
+	TArray<FGuid> ObjectBindings;
+	ObjectBindings.Add(ObjectBinding);
+
+	AddAnimationSubMenu(MenuBuilder, ObjectBindings, Skeleton, Track);
+
+	return MenuBuilder.MakeWidget();
+}
+
+void FStitchAnimTrackEditor::AddAnimationSubMenu(FMenuBuilder& MenuBuilder, TArray<FGuid> ObjectBindings, USkeleton* Skeleton, UMovieSceneTrack* Track)
+{
+	using namespace UE::PoseSearch;
+	TSharedPtr<ISequencer> SequencerPtr = GetSequencer();
+	UMovieSceneSequence* Sequence = SequencerPtr.IsValid() ? SequencerPtr->GetFocusedMovieSceneSequence() : nullptr;
+
+	FAssetPickerConfig AssetPickerConfig;
+	{
+		AssetPickerConfig.OnAssetSelected = FOnAssetSelected::CreateRaw(this, &FStitchAnimTrackEditor::OnAnimationDatabaseAssetSelected, ObjectBindings, Track);
+		AssetPickerConfig.OnAssetEnterPressed = FOnAssetEnterPressed::CreateRaw(this, &FStitchAnimTrackEditor::OnAnimationDatabaseAssetEnterPressed, ObjectBindings, Track);
+		AssetPickerConfig.bAllowNullSelection = false;
+		AssetPickerConfig.bAddFilterUI = true;
+		AssetPickerConfig.bShowTypeInColumnView = false;
+		AssetPickerConfig.InitialAssetViewType = EAssetViewType::List;
+		AssetPickerConfig.Filter.bRecursiveClasses = true;
+		AssetPickerConfig.Filter.ClassPaths.Add(UPoseSearchDatabase::StaticClass()->GetClassPathName());
+		AssetPickerConfig.SaveSettingsName = TEXT("SequencerAssetPicker");
+		AssetPickerConfig.AdditionalReferencingAssets.Add(FAssetData(Sequence));
+	}
+
+	FContentBrowserModule& ContentBrowserModule = FModuleManager::Get().LoadModuleChecked<FContentBrowserModule>(TEXT("ContentBrowser"));
+
+	const float WidthOverride = SequencerPtr.IsValid() ? SequencerPtr->GetSequencerSettings()->GetAssetBrowserWidth() : 500.f;
+	const float HeightOverride = SequencerPtr.IsValid() ? SequencerPtr->GetSequencerSettings()->GetAssetBrowserHeight() : 400.f;
+
+	TSharedPtr<SBox> MenuEntry = SNew(SBox)
+		.WidthOverride(WidthOverride)
+		.HeightOverride(HeightOverride)
+		[
+			ContentBrowserModule.Get().CreateAssetPicker(AssetPickerConfig)
+		];
+
+	MenuBuilder.AddWidget(MenuEntry.ToSharedRef(), FText::GetEmpty(), true);
+}
+
+void FStitchAnimTrackEditor::OnAnimationDatabaseAssetSelected(const FAssetData& AssetData, TArray<FGuid> ObjectBindings, UMovieSceneTrack* Track)
+{
+	using namespace UE::PoseSearch;
+	FSlateApplication::Get().DismissAllMenus();
+
+	UObject* SelectedObject = AssetData.GetAsset();
+	TSharedPtr<ISequencer> SequencerPtr = GetSequencer();
+
+	if (SelectedObject && SelectedObject->IsA(UPoseSearchDatabase::StaticClass()) && SequencerPtr.IsValid())
+	{
+		UPoseSearchDatabase* PoseSearchDatabase = CastChecked<UPoseSearchDatabase>(AssetData.GetAsset());
+
+		const FScopedTransaction Transaction(LOCTEXT("AddStitchAnim_Transaction", "Add Stitch Animation"));
+
+		for (FGuid ObjectBinding : ObjectBindings)
+		{
+			UObject* Object = SequencerPtr->FindSpawnedObjectOrTemplate(ObjectBinding);
+			int32 RowIndex = INDEX_NONE;
+			AnimatablePropertyChanged(FOnKeyProperty::CreateRaw(this, &FStitchAnimTrackEditor::AddKeyInternal, Object, PoseSearchDatabase, Track, RowIndex));
+		}
+	}
+}
+
+void FStitchAnimTrackEditor::OnAnimationDatabaseAssetEnterPressed(const TArray<FAssetData>& AssetData, TArray<FGuid> ObjectBindings, UMovieSceneTrack* Track)
+{
+	if (AssetData.Num() > 0)
+	{
+		OnAnimationDatabaseAssetSelected(AssetData[0].GetAsset(), ObjectBindings, Track);
+	}
+}
+
+
+FKeyPropertyResult FStitchAnimTrackEditor::AddKeyInternal(FFrameNumber KeyTime, UObject* Object, UPoseSearchDatabase* PoseSearchDatabase, UMovieSceneTrack* Track, int32 RowIndex)
+{
+	FKeyPropertyResult KeyPropertyResult;
+
+	FFindOrCreateHandleResult HandleResult = FindOrCreateHandleToObject(Object);
+	FGuid ObjectHandle = HandleResult.Handle;
+	KeyPropertyResult.bHandleCreated |= HandleResult.bWasCreated;
+	if (ObjectHandle.IsValid())
+	{
+		UMovieScene* MovieScene = GetSequencer()->GetFocusedMovieSceneSequence()->GetMovieScene();
+		UMovieSceneStitchAnimTrack* StitchAnimTrack = Cast<UMovieSceneStitchAnimTrack>(Track);
+		FMovieSceneBinding* Binding = MovieScene->FindBinding(ObjectHandle);
+
+		// Add a track if no track was specified or if the track specified doesn't belong to the tracks of the targeted guid
+		if (!StitchAnimTrack || (Binding && !Binding->GetTracks().Contains(StitchAnimTrack)))
+		{
+			StitchAnimTrack = CastChecked<UMovieSceneStitchAnimTrack>(AddTrack(MovieScene, ObjectHandle, UMovieSceneStitchAnimTrack::StaticClass(), NAME_None), ECastCheckedType::NullAllowed);
+			KeyPropertyResult.bTrackCreated = true;
+		}
+
+		if (ensure(StitchAnimTrack))
+		{
+			StitchAnimTrack->Modify();
+
+			UMovieSceneStitchAnimSection* NewSection = Cast<UMovieSceneStitchAnimSection>(StitchAnimTrack->AddNewAnimationOnRow(KeyTime, PoseSearchDatabase, RowIndex));
+			KeyPropertyResult.bTrackModified = true;
+			KeyPropertyResult.SectionsCreated.Add(NewSection);
+
+			GetSequencer()->EmptySelection();
+			GetSequencer()->SelectSection(NewSection);
+			GetSequencer()->ThrobSectionSelection();
+		}
+	}
+
+	return KeyPropertyResult;
+}
+
+TSharedPtr<SWidget> FStitchAnimTrackEditor::BuildOutlinerEditWidget(const FGuid& ObjectBinding, UMovieSceneTrack* Track, const FBuildEditWidgetParams& Params)
+{
+	USkeleton* Skeleton = AcquireSkeletonFromObjectGuid(ObjectBinding, GetSequencer());
+
+	if (Skeleton)
+	{
+		FOnGetContent HandleGetAddButtonContent = FOnGetContent::CreateSP(this, &FStitchAnimTrackEditor::BuildAddAnimationSubMenu, ObjectBinding, Skeleton, Params.TrackModel.AsWeak());
+		return UE::Sequencer::MakeAddButton(LOCTEXT("AnimationText", "Stitch Animation"), HandleGetAddButtonContent, Params.ViewModel);
+	}
+	else
+	{
+		return TSharedPtr<SWidget>();
+	}
+}
+
+bool FStitchAnimTrackEditor::OnAllowDrop(const FDragDropEvent& DragDropEvent, FSequencerDragDropParams& DragDropParams)
+{
+	using namespace UE::PoseSearch;
+	TSharedPtr<FDragDropOperation> Operation = DragDropEvent.GetOperation();
+
+	if (!Operation.IsValid() || !Operation->IsOfType<FAssetDragDropOp>())
+	{
+		return false;
+	}
+
+	if (!DragDropParams.TargetObjectGuid.IsValid())
+	{
+		return false;
+	}
+
+	TSharedPtr<ISequencer> SequencerPtr = GetSequencer();
+	if (!SequencerPtr)
+	{
+		return false;
+	}
+
+	UMovieSceneSequence* FocusedSequence = SequencerPtr->GetFocusedMovieSceneSequence();
+	if (!FocusedSequence)
+	{
+		return false;
+	}
+
+	TArray<USkeletalMeshComponent*> SkeletalMeshComponents = AcquireSkeletalMeshComponentsFromObjectGuid(DragDropParams.TargetObjectGuid, SequencerPtr, false);
+
+	TSharedPtr<FAssetDragDropOp> DragDropOp = StaticCastSharedPtr<FAssetDragDropOp>(Operation);
+
+	for (const FAssetData& AssetData : DragDropOp->GetAssets())
+	{
+		if (!MovieSceneToolHelpers::IsValidAsset(FocusedSequence, AssetData))
+		{
+			continue;
+		}
+
+		UPoseSearchDatabase* PoseSearchDatabase = Cast<UPoseSearchDatabase>(AssetData.GetAsset());
+
+		for (USkeletalMeshComponent* SkeletalMeshComponent : SkeletalMeshComponents)
+		{
+			USkeleton* Skeleton = GetSkeletonFromComponent(SkeletalMeshComponent);
+			if (PoseSearchDatabase != nullptr)
+			//if (bValidAnimSequence && Skeleton && Skeleton->IsCompatibleForEditor(AnimSequence->GetSkeleton()))
+			{
+				FFrameRate TickResolution = SequencerPtr->GetFocusedTickResolution();
+				FFrameNumber LengthInFrames = TickResolution.AsFrameNumber(3.0f);
+				DragDropParams.FrameRange = TRange<FFrameNumber>(DragDropParams.FrameNumber, DragDropParams.FrameNumber + LengthInFrames);
+				return true;
+			}
+		}
+	}
+
+	return false;
+}
+
+
+FReply FStitchAnimTrackEditor::OnDrop(const FDragDropEvent& DragDropEvent, const FSequencerDragDropParams& DragDropParams)
+{
+	using namespace UE::PoseSearch;
+	TSharedPtr<FDragDropOperation> Operation = DragDropEvent.GetOperation();
+
+	if (!Operation.IsValid() || !Operation->IsOfType<FAssetDragDropOp>())
+	{
+		return FReply::Unhandled();
+	}
+
+	if (!DragDropParams.TargetObjectGuid.IsValid())
+	{
+		return FReply::Unhandled();
+	}
+
+	TSharedPtr<ISequencer> SequencerPtr = GetSequencer();
+	if (!SequencerPtr)
+	{
+		return FReply::Unhandled();
+	}
+
+	UMovieSceneSequence* FocusedSequence = SequencerPtr->GetFocusedMovieSceneSequence();
+	if (!FocusedSequence)
+	{
+		return FReply::Unhandled();
+	}
+
+	TArray<USkeletalMeshComponent*> SkeletalMeshComponents = AcquireSkeletalMeshComponentsFromObjectGuid(DragDropParams.TargetObjectGuid, SequencerPtr, false);
+
+	const FScopedTransaction Transaction(LOCTEXT("DropAssets", "Drop Assets"));
+
+	TSharedPtr<FAssetDragDropOp> DragDropOp = StaticCastSharedPtr<FAssetDragDropOp>(Operation);
+
+	FMovieSceneTrackEditor::BeginKeying(DragDropParams.FrameNumber);
+
+	bool bAnyDropped = false;
+	for (const FAssetData& AssetData : DragDropOp->GetAssets())
+	{
+		if (!MovieSceneToolHelpers::IsValidAsset(FocusedSequence, AssetData))
+		{
+			continue;
+		}
+
+		UPoseSearchDatabase* PoseSearchDatabase = Cast<UPoseSearchDatabase>(AssetData.GetAsset());
+		for (USkeletalMeshComponent* SkeletalMeshComponent : SkeletalMeshComponents)
+		{
+			USkeleton* Skeleton = GetSkeletonFromComponent(SkeletalMeshComponent);
+
+			if (PoseSearchDatabase != nullptr)// && Skeleton && Skeleton->IsCompatibleForEditor(AnimSequence->GetSkeleton()))
+			{
+				UObject* BoundObject = SequencerPtr.IsValid() ? SequencerPtr->FindSpawnedObjectOrTemplate(DragDropParams.TargetObjectGuid) : nullptr;
+
+				AnimatablePropertyChanged(FOnKeyProperty::CreateRaw(this, &FStitchAnimTrackEditor::AddKeyInternal, BoundObject, PoseSearchDatabase, DragDropParams.Track.Get(), DragDropParams.RowIndex));
+
+				bAnyDropped = true;
+			}
+		}
+	}
+
+	FMovieSceneTrackEditor::EndKeying();
+
+	return bAnyDropped ? FReply::Handled() : FReply::Unhandled();
+}
+
+
+#undef LOCTEXT_NAMESPACE
diff --git a/Engine/Plugins/Experimental/MovieScenePoseSearchTracks/Source/MovieScenePoseSearchTracksEditor/Private/TrackEditors/StitchAnimTrackEditor.h b/Engine/Plugins/Experimental/MovieScenePoseSearchTracks/Source/MovieScenePoseSearchTracksEditor/Private/TrackEditors/StitchAnimTrackEditor.h
new file mode 100644
--- /dev/null
+++ b/Engine/Plugins/Experimental/MovieScenePoseSearchTracks/Source/MovieScenePoseSearchTracksEditor/Private/TrackEditors/StitchAnimTrackEditor.h
@@ -0,0 +1,117 @@
+// Copyright Epic Games, Inc. All Rights Reserved.
+
+#pragma once
+
+#include "CoreTypes.h"
+#include "Misc/Guid.h"
+#include "Templates/SubclassOf.h"
+#include "Templates/UniquePtr.h"
+#include "Widgets/SWidget.h"
+#include "ISequencer.h"
+#include "MovieSceneTrack.h"
+#include "ISequencerSection.h"
+#include "ISequencerTrackEditor.h"
+#include "MovieSceneTrackEditor.h"
+#include "SequencerCoreFwd.h"
+
+namespace UE::Sequencer
+{
+	class ITrackExtension;
+}
+
+struct FAssetData;
+class FMenuBuilder;
+class FSequencerSectionPainter;
+class UMovieSceneStitchAnimSection;
+class USkeleton;
+class UPoseSearchDatabase;
+
+/**
+ * Tools for stitch anim tracks
+ */
+class FStitchAnimTrackEditor : public FMovieSceneTrackEditor
+{
+public:
+	/**
+	 * Constructor
+	 *
+	 * @param InSequencer The sequencer instance to be used by this tool
+	 */
+	FStitchAnimTrackEditor(TSharedRef<ISequencer> InSequencer);
+
+	/** Virtual destructor. */
+	virtual ~FStitchAnimTrackEditor() { }
+
+	/**
+	 * Creates an instance of this class.  Called by a sequencer
+	 *
+	 * @param OwningSequencer The sequencer instance to be used by this tool
+	 * @return The new instance of this class
+	 */
+	static TSharedRef<ISequencerTrackEditor> CreateTrackEditor(TSharedRef<ISequencer> OwningSequencer);
+
+public:
+
+	// ISequencerTrackEditor interface
+
+	virtual void BuildObjectBindingTrackMenu(FMenuBuilder& MenuBuilder, const TArray<FGuid>& ObjectBindings, const UClass* ObjectClass) override;
+	virtual bool HandleAssetAdded(UObject* Asset, const FGuid& TargetObjectGuid) override;
+	virtual TSharedRef<ISequencerSection> MakeSectionInterface(UMovieSceneSection& SectionObject, UMovieSceneTrack& Track, FGuid ObjectBinding) override;
+	virtual bool SupportsSequence(UMovieSceneSequence* InSequence) const override;
+	virtual bool SupportsType(TSubclassOf<UMovieSceneTrack> Type) const override;
+	virtual TSharedPtr<SWidget> BuildOutlinerEditWidget(const FGuid& ObjectBinding, UMovieSceneTrack* Track, const FBuildEditWidgetParams& Params) override;
+	virtual bool OnAllowDrop(const FDragDropEvent& DragDropEvent, FSequencerDragDropParams& DragDropParams) override;
+	virtual FReply OnDrop(const FDragDropEvent& DragDropEvent, const FSequencerDragDropParams& DragDropParams) override;
+
+private:
+
+	/** Animation sub menu */
+	TSharedRef<SWidget> BuildAddAnimationSubMenu(FGuid ObjectBinding, USkeleton* Skeleton, UE::Sequencer::TWeakViewModelPtr<UE::Sequencer::ITrackExtension> TrackModel);
+	TSharedRef<SWidget> BuildAnimationSubMenu(FGuid ObjectBinding, USkeleton* Skeleton, UMovieSceneTrack* Track);
+	void AddAnimationSubMenu(FMenuBuilder& MenuBuilder, TArray<FGuid> ObjectBindings, USkeleton* Skeleton, UMovieSceneTrack* Track);
+
+	/** Animation asset selected */
+	void OnAnimationDatabaseAssetSelected(const FAssetData& AssetData, TArray<FGuid> ObjectBindings, UMovieSceneTrack* Track);
+
+	/** Animation asset enter pressed */
+	void OnAnimationDatabaseAssetEnterPressed(const TArray<FAssetData>& AssetData, TArray<FGuid> ObjectBindings, UMovieSceneTrack* Track);
+
+	/** Delegate for AnimatablePropertyChanged in AddKey */
+	FKeyPropertyResult AddKeyInternal(FFrameNumber KeyTime, UObject* Object, UPoseSearchDatabase* PoseSearchDatabase, UMovieSceneTrack* Track, int32 RowIndex);
+};
+
+
+/** Class for stitch anim sections */
+class FStitchAnimSection
+	: public ISequencerSection
+	, public TSharedFromThis<FStitchAnimSection>
+{
+public:
+
+	/** Constructor. */
+	FStitchAnimSection(UMovieSceneSection& InSection, TWeakPtr<ISequencer> InSequencer);
+
+	/** Virtual destructor. */
+	virtual ~FStitchAnimSection();
+
+public:
+
+	// ISequencerSection interface
+
+	virtual UMovieSceneSection* GetSectionObject() override;
+	virtual FText GetSectionTitle() const override;
+	virtual FText GetSectionToolTip() const override;
+	virtual float GetSectionHeight(const UE::Sequencer::FViewDensityInfo& ViewDensity) const override;
+	virtual FMargin GetContentPadding() const override;
+	virtual int32 OnPaintSection(FSequencerSectionPainter& Painter) const override;
+
+private:
+
+	/** The section we are visualizing */
+	UMovieSceneStitchAnimSection& Section;
+
+	/** Used to draw animation frame, need selection state and local time*/
+	TWeakPtr<ISequencer> Sequencer;
+
+	TUniquePtr<FMovieSceneSequenceTransform> InitialDragTransform;
+};
diff --git a/Engine/Source/Runtime/MovieSceneTracks/Private/Systems/MovieSceneTransformOriginSystem.cpp b/Engine/Source/Runtime/MovieSceneTracks/Private/Systems/MovieSceneTransformOriginSystem.cpp
--- a/Engine/Source/Runtime/MovieSceneTracks/Private/Systems/MovieSceneTransformOriginSystem.cpp
+++ b/Engine/Source/Runtime/MovieSceneTracks/Private/Systems/MovieSceneTransformOriginSystem.cpp
@@ -315,9 +315,7 @@ void UMovieSceneTransformOriginSystem::OnSchedulePersistentTasks(UE::MovieScene:
 	AssignFilter.Combine(LocationAndRotationFilterResults);
 	
 	InstanceHandleToParentHandle.Empty();
-#if WITH_EDITOR
 	SequenceIDToInstanceHandle.Empty();
-#endif
 	
 	FEntityTaskBuilder()
 	.Read(BuiltInComponents->RootInstanceHandle)
@@ -329,9 +327,7 @@ void UMovieSceneTransformOriginSystem::OnSchedulePersistentTasks(UE::MovieScene:
 	{
 		FInstanceHandle SubInstanceHandle = InstanceRegistry->GetInstance(RootInstance).FindSubInstance(SequenceID);
 		InstanceHandleToParentHandle.AddUnique(FInstanceToParentPair(SubInstanceHandle, Instance));
-#if WITH_EDITOR
 		SequenceIDToInstanceHandle.Add(SequenceID, SubInstanceHandle);
-#endif
 		
 	});
 
diff --git a/Engine/Source/Runtime/MovieSceneTracks/Public/Systems/MovieSceneTransformOriginSystem.h b/Engine/Source/Runtime/MovieSceneTracks/Public/Systems/MovieSceneTransformOriginSystem.h
--- a/Engine/Source/Runtime/MovieSceneTracks/Public/Systems/MovieSceneTransformOriginSystem.h
+++ b/Engine/Source/Runtime/MovieSceneTracks/Public/Systems/MovieSceneTransformOriginSystem.h
@@ -63,6 +63,12 @@ public:
 
 	UMovieSceneTransformOriginSystem(const FObjectInitializer& ObjInit);
 
+public:
+	const TSparseArray<FTransform>& GetTransformOriginsByInstanceID() const { return TransformOriginsByInstanceID; }
+
+
+	const TMap<FMovieSceneSequenceID, UE::MovieScene::FInstanceHandle>& GetSequenceIDToInstanceHandle() const { return SequenceIDToInstanceHandle; }
+
 private:
 
 	virtual void OnLink() override;
@@ -76,18 +82,9 @@ private:
 	TArray<FInstanceToParentPair> InstanceHandleToParentHandle;
 	UE::MovieScene::FEntityComponentFilter LocationAndRotationFilterResults;
 
-#if WITH_EDITOR
-	TMap<FMovieSceneSequenceID, UE::MovieScene::FInstanceHandle> SequenceIDToInstanceHandle;
-#endif
-	
 
-public:
-	const TSparseArray<FTransform>& GetTransformOriginsByInstanceID() const { return TransformOriginsByInstanceID; }
+	TMap<FMovieSceneSequenceID, UE::MovieScene::FInstanceHandle> SequenceIDToInstanceHandle;
 
-#if WITH_EDITOR
-	const TMap<FMovieSceneSequenceID, UE::MovieScene::FInstanceHandle>& GetSequenceIDToInstanceHandle() const { return SequenceIDToInstanceHandle; }
-#endif
-	
 };
 
 
