Sequencer: Use static_cast<int32> instead of int

diff --git a/Engine/Plugins/Experimental/ChaosCaching/Source/ChaosCaching/Private/Chaos/Sequencer/MovieSceneChaosCacheTrack.cpp b/Engine/Plugins/Experimental/ChaosCaching/Source/ChaosCaching/Private/Chaos/Sequencer/MovieSceneChaosCacheTrack.cpp
--- a/Engine/Plugins/Experimental/ChaosCaching/Source/ChaosCaching/Private/Chaos/Sequencer/MovieSceneChaosCacheTrack.cpp
+++ b/Engine/Plugins/Experimental/ChaosCaching/Source/ChaosCaching/Private/Chaos/Sequencer/MovieSceneChaosCacheTrack.cpp
@@ -36,7 +36,7 @@ UMovieSceneSection* UMovieSceneChaosCacheTrack::AddNewAnimation(FFrameNumber Key
 		if(ChaosCache && ChaosCache->CacheCollection)
 		{
 			const FFrameTime AnimationLength = ChaosCache->CacheCollection->GetMaxDuration() * GetTypedOuter<UMovieScene>()->GetTickResolution();
-			const int32 IFrameNumber = AnimationLength.FrameNumber.Value + (int)(AnimationLength.GetSubFrame() + 0.5f) + 1;
+			const int32 IFrameNumber = AnimationLength.FrameNumber.Value + static_cast<int32>(AnimationLength.GetSubFrame() + 0.5f) + 1;
 			NewSection->InitialPlacementOnRow(AnimationSections, KeyTime, IFrameNumber, INDEX_NONE);
 			
 			NewSection->Params.CacheCollection = ChaosCache->CacheCollection;
diff --git a/Engine/Plugins/Experimental/MovieScenePoseSearchTracks/Source/MovieScenePoseSearchTracks/Private/Tracks/MovieSceneStitchAnimTrack.cpp b/Engine/Plugins/Experimental/MovieScenePoseSearchTracks/Source/MovieScenePoseSearchTracks/Private/Tracks/MovieSceneStitchAnimTrack.cpp
--- a/Engine/Plugins/Experimental/MovieScenePoseSearchTracks/Source/MovieScenePoseSearchTracks/Private/Tracks/MovieSceneStitchAnimTrack.cpp
+++ b/Engine/Plugins/Experimental/MovieScenePoseSearchTracks/Source/MovieScenePoseSearchTracks/Private/Tracks/MovieSceneStitchAnimTrack.cpp
@@ -82,7 +82,7 @@ UMovieSceneSection* UMovieSceneStitchAnimTrack::AddNewAnimationOnRow(FFrameNumbe
 	UMovieSceneStitchAnimSection* NewSection = Cast<UMovieSceneStitchAnimSection>(CreateNewSection());
 	{
 		FFrameTime AnimationLength = 5.0f * GetTypedOuter<UMovieScene>()->GetTickResolution();
-		int32 IFrameNumber = AnimationLength.FrameNumber.Value + (int)(AnimationLength.GetSubFrame() + 0.5f) + 1;
+		int32 IFrameNumber = AnimationLength.FrameNumber.Value + static_cast<int32>(AnimationLength.GetSubFrame() + 0.5f) + 1;
 
 		NewSection->InitialPlacementOnRow(AnimationSections, KeyTime, IFrameNumber, RowIndex);
 		NewSection->StitchDatabase = PoseSearchDatabase;
diff --git a/Engine/Plugins/FX/NiagaraSimCaching/Source/NiagaraSimCaching/Private/Niagara/Sequencer/MovieSceneNiagaraCacheTrack.cpp b/Engine/Plugins/FX/NiagaraSimCaching/Source/NiagaraSimCaching/Private/Niagara/Sequencer/MovieSceneNiagaraCacheTrack.cpp
--- a/Engine/Plugins/FX/NiagaraSimCaching/Source/NiagaraSimCaching/Private/Niagara/Sequencer/MovieSceneNiagaraCacheTrack.cpp
+++ b/Engine/Plugins/FX/NiagaraSimCaching/Source/NiagaraSimCaching/Private/Niagara/Sequencer/MovieSceneNiagaraCacheTrack.cpp
@@ -34,7 +34,7 @@ UMovieSceneSection* UMovieSceneNiagaraCacheTrack::AddNewAnimation(FFrameNumber K
 	UMovieSceneNiagaraCacheSection* NewSection = Cast<UMovieSceneNiagaraCacheSection>(CreateNewSection());
 	{
 		//const FFrameTime AnimationLength = NiagaraCache->CacheCollection->GetMaxDuration() * GetTypedOuter<UMovieScene>()->GetTickResolution();
-		//const int32 IFrameNumber = AnimationLength.FrameNumber.Value + (int)(AnimationLength.GetSubFrame() + 0.5f) + 1;
+		//const int32 IFrameNumber = AnimationLength.FrameNumber.Value + static_cast<int32>(AnimationLength.GetSubFrame() + 0.5f) + 1;
 		//NewSection->InitialPlacementOnRow(AnimationSections, KeyTime, IFrameNumber, INDEX_NONE);
 	}
 
diff --git a/Engine/Plugins/Runtime/GeometryCache/Source/GeometryCacheTracks/Private/MovieSceneGeometryCacheSection.cpp b/Engine/Plugins/Runtime/GeometryCache/Source/GeometryCacheTracks/Private/MovieSceneGeometryCacheSection.cpp
--- a/Engine/Plugins/Runtime/GeometryCache/Source/GeometryCacheTracks/Private/MovieSceneGeometryCacheSection.cpp
+++ b/Engine/Plugins/Runtime/GeometryCache/Source/GeometryCacheTracks/Private/MovieSceneGeometryCacheSection.cpp
@@ -118,7 +118,7 @@ TOptional<TRange<FFrameNumber> > UMovieSceneGeometryCacheSection::GetAutoSizeRan
 {
 	FFrameRate FrameRate = GetTypedOuter<UMovieScene>()->GetTickResolution();
 	FFrameTime AnimationLength = Params.GetSequenceLength() * FrameRate;
-	int32 IFrameNumber = AnimationLength.FrameNumber.Value + (int)(AnimationLength.GetSubFrame() + 0.5f);
+	int32 IFrameNumber = AnimationLength.FrameNumber.Value + static_cast<int32>(AnimationLength.GetSubFrame() + 0.5f);
 
 	return TRange<FFrameNumber>(GetInclusiveStartFrame(), GetInclusiveStartFrame() + IFrameNumber + 1);
 }
diff --git a/Engine/Plugins/Runtime/GeometryCache/Source/GeometryCacheTracks/Private/MovieSceneGeometryCacheTrack.cpp b/Engine/Plugins/Runtime/GeometryCache/Source/GeometryCacheTracks/Private/MovieSceneGeometryCacheTrack.cpp
--- a/Engine/Plugins/Runtime/GeometryCache/Source/GeometryCacheTracks/Private/MovieSceneGeometryCacheTrack.cpp
+++ b/Engine/Plugins/Runtime/GeometryCache/Source/GeometryCacheTracks/Private/MovieSceneGeometryCacheTrack.cpp
@@ -36,7 +36,7 @@ UMovieSceneSection* UMovieSceneGeometryCacheTrack::AddNewAnimation(FFrameNumber
 	UMovieSceneGeometryCacheSection* NewSection = Cast<UMovieSceneGeometryCacheSection>(CreateNewSection());
 	{
 		FFrameTime AnimationLength = GeomCacheComp->GetDuration()* GetTypedOuter<UMovieScene>()->GetTickResolution();
-		int32 IFrameNumber = AnimationLength.FrameNumber.Value + (int)(AnimationLength.GetSubFrame() + 0.5f) + 1;
+		int32 IFrameNumber = AnimationLength.FrameNumber.Value + static_cast<int32>(AnimationLength.GetSubFrame() + 0.5f) + 1;
 		NewSection->InitialPlacementOnRow(AnimationSections, KeyTime, IFrameNumber, INDEX_NONE);
 		
 		NewSection->Params.GeometryCacheAsset = (GeomCacheComp->GetGeometryCache());
diff --git a/Engine/Plugins/Runtime/HairStrands/Source/HairStrandsCore/Private/MovieSceneGroomCacheSection.cpp b/Engine/Plugins/Runtime/HairStrands/Source/HairStrandsCore/Private/MovieSceneGroomCacheSection.cpp
--- a/Engine/Plugins/Runtime/HairStrands/Source/HairStrandsCore/Private/MovieSceneGroomCacheSection.cpp
+++ b/Engine/Plugins/Runtime/HairStrands/Source/HairStrandsCore/Private/MovieSceneGroomCacheSection.cpp
@@ -38,7 +38,7 @@ TOptional<TRange<FFrameNumber> > UMovieSceneGroomCacheSection::GetAutoSizeRange(
 {
 	FFrameRate FrameRate = GetTypedOuter<UMovieScene>()->GetTickResolution();
 	FFrameTime AnimationLength = Params.GetSequenceLength() * FrameRate;
-	int32 IFrameNumber = AnimationLength.FrameNumber.Value + (int)(AnimationLength.GetSubFrame() + 0.5f);
+	int32 IFrameNumber = AnimationLength.FrameNumber.Value + static_cast<int32>(AnimationLength.GetSubFrame() + 0.5f);
 
 	return TRange<FFrameNumber>(GetInclusiveStartFrame(), GetInclusiveStartFrame() + IFrameNumber + 1);
 }
diff --git a/Engine/Plugins/Runtime/HairStrands/Source/HairStrandsCore/Private/MovieSceneGroomCacheTrack.cpp b/Engine/Plugins/Runtime/HairStrands/Source/HairStrandsCore/Private/MovieSceneGroomCacheTrack.cpp
--- a/Engine/Plugins/Runtime/HairStrands/Source/HairStrandsCore/Private/MovieSceneGroomCacheTrack.cpp
+++ b/Engine/Plugins/Runtime/HairStrands/Source/HairStrandsCore/Private/MovieSceneGroomCacheTrack.cpp
@@ -34,7 +34,7 @@ UMovieSceneSection* UMovieSceneGroomCacheTrack::AddNewAnimation(FFrameNumber Key
 	UMovieSceneGroomCacheSection* NewSection = Cast<UMovieSceneGroomCacheSection>(CreateNewSection());
 	{
 		FFrameTime AnimationLength = GroomComp->GetGroomCacheDuration() * GetTypedOuter<UMovieScene>()->GetTickResolution();
-		int32 IFrameNumber = AnimationLength.FrameNumber.Value + (int)(AnimationLength.GetSubFrame() + 0.5f) + 1;
+		int32 IFrameNumber = AnimationLength.FrameNumber.Value + static_cast<int32>(AnimationLength.GetSubFrame() + 0.5f) + 1;
 		NewSection->InitialPlacementOnRow(AnimationSections, KeyTime, IFrameNumber, INDEX_NONE);
 		
 		NewSection->Params.GroomCache = GroomComp->GetGroomCache();
diff --git a/Engine/Source/Runtime/MovieSceneTracks/Private/Sections/MovieSceneAudioSection.cpp b/Engine/Source/Runtime/MovieSceneTracks/Private/Sections/MovieSceneAudioSection.cpp
--- a/Engine/Source/Runtime/MovieSceneTracks/Private/Sections/MovieSceneAudioSection.cpp
+++ b/Engine/Source/Runtime/MovieSceneTracks/Private/Sections/MovieSceneAudioSection.cpp
@@ -370,7 +370,7 @@ TOptional<TRange<FFrameNumber> > UMovieSceneAudioSection::GetAutoSizeRange() con
 		DurationToUse = FMath::Max(SoundDuration * FrameRate - StartFrameOffset, FFrameTime(1));
 	}
 
-	const int32 IFrameNumber = DurationToUse.FrameNumber.Value + (int)(DurationToUse.GetSubFrame() + 0.5f);
+	const int32 IFrameNumber = DurationToUse.FrameNumber.Value + static_cast<int32>(DurationToUse.GetSubFrame() + 0.5f);
 	return TRange<FFrameNumber>(GetInclusiveStartFrame(), GetInclusiveStartFrame() + IFrameNumber);
 }
 
diff --git a/Engine/Source/Runtime/MovieSceneTracks/Private/Tracks/MovieSceneSkeletalAnimationTrack.cpp b/Engine/Source/Runtime/MovieSceneTracks/Private/Tracks/MovieSceneSkeletalAnimationTrack.cpp
--- a/Engine/Source/Runtime/MovieSceneTracks/Private/Tracks/MovieSceneSkeletalAnimationTrack.cpp
+++ b/Engine/Source/Runtime/MovieSceneTracks/Private/Tracks/MovieSceneSkeletalAnimationTrack.cpp
@@ -57,7 +57,7 @@ UMovieSceneSection* UMovieSceneSkeletalAnimationTrack::AddNewAnimationOnRow(FFra
 	UMovieSceneSkeletalAnimationSection* NewSection = Cast<UMovieSceneSkeletalAnimationSection>(CreateNewSection());
 	{
 		FFrameTime AnimationLength = AnimSequence->GetPlayLength() * GetTypedOuter<UMovieScene>()->GetTickResolution();
-		int32 IFrameNumber = AnimationLength.FrameNumber.Value + (int)(AnimationLength.GetSubFrame() + 0.5f) + 1;
+		int32 IFrameNumber = AnimationLength.FrameNumber.Value + static_cast<int32>(AnimationLength.GetSubFrame() + 0.5f) + 1;
 		NewSection->InitialPlacementOnRow(AnimationSections, KeyTime, IFrameNumber, RowIndex);
 		NewSection->Params.Animation = AnimSequence;
 
