Fixed uninitialized member variable

diff --git a/Engine/Plugins/Experimental/MovieScenePoseSearchTracks/Source/MovieScenePoseSearchTracks/Public/Sections/MovieSceneStitchAnimSection.h b/Engine/Plugins/Experimental/MovieScenePoseSearchTracks/Source/MovieScenePoseSearchTracks/Public/Sections/MovieSceneStitchAnimSection.h
--- a/Engine/Plugins/Experimental/MovieScenePoseSearchTracks/Source/MovieScenePoseSearchTracks/Public/Sections/MovieSceneStitchAnimSection.h
+++ b/Engine/Plugins/Experimental/MovieScenePoseSearchTracks/Source/MovieScenePoseSearchTracks/Public/Sections/MovieSceneStitchAnimSection.h
@@ -86,7 +86,7 @@ struct FMovieSceneStitchAnimEvaluationTask : public FAnimNextEvaluationTask
 	mutable TObjectPtr<UAnimationAsset> MatchedAsset;
 
 	UPROPERTY()
-	mutable float MatchedAssetTime;
+	mutable float MatchedAssetTime = 0.f;
 
 	UPROPERTY()
 	mutable float MatchedAssetPlayRate = 1.0f;
