Sequencer: Added an option to stitch tracks for specifying what space the target transform is in

diff --git a/Engine/Plugins/Experimental/MovieScenePoseSearchTracks/Source/MovieScenePoseSearchTracks/Private/Sections/MovieSceneStitchAnimSection.cpp b/Engine/Plugins/Experimental/MovieScenePoseSearchTracks/Source/MovieScenePoseSearchTracks/Private/Sections/MovieSceneStitchAnimSection.cpp
--- a/Engine/Plugins/Experimental/MovieScenePoseSearchTracks/Source/MovieScenePoseSearchTracks/Private/Sections/MovieSceneStitchAnimSection.cpp
+++ b/Engine/Plugins/Experimental/MovieScenePoseSearchTracks/Source/MovieScenePoseSearchTracks/Private/Sections/MovieSceneStitchAnimSection.cpp
@@ -259,6 +259,7 @@ void UMovieSceneStitchAnimSection::ImportEntityImpl(UMovieSceneEntitySystemLinke
 	StitchData.TargetTransform = TargetTransform;
 	StitchData.StartFrame = GetInclusiveStartFrame();
 	StitchData.EndFrame = GetExclusiveEndFrame();
+	StitchData.TargetTransformSpace = TargetTransformSpace;
 	
 	// Make a default eval task. This will be filled out more by the stitch system
 	TSharedPtr<FMovieSceneStitchAnimEvaluationTask> EvalTask = MakeShared<FMovieSceneStitchAnimEvaluationTask>();
diff --git a/Engine/Plugins/Experimental/MovieScenePoseSearchTracks/Source/MovieScenePoseSearchTracks/Private/Systems/MovieSceneStitchAnimSystem.cpp b/Engine/Plugins/Experimental/MovieScenePoseSearchTracks/Source/MovieScenePoseSearchTracks/Private/Systems/MovieSceneStitchAnimSystem.cpp
--- a/Engine/Plugins/Experimental/MovieScenePoseSearchTracks/Source/MovieScenePoseSearchTracks/Private/Systems/MovieSceneStitchAnimSystem.cpp
+++ b/Engine/Plugins/Experimental/MovieScenePoseSearchTracks/Source/MovieScenePoseSearchTracks/Private/Systems/MovieSceneStitchAnimSystem.cpp
@@ -178,7 +178,9 @@ namespace UE::MovieScene
 					// We use the previous time for the initial time, as it will match the transform from the actor which hasn't
 					// yet been set for this frame.
 					AnimTask->InitialTime = PreviousTimeSeconds;
-					AnimTask->SequenceTransformOrigin = GetTransformOrigin(InstanceHandle);
+					AnimTask->SequenceTransformOrigin = StitchAnim.TargetTransformSpace == EMovieSceneRootMotionSpace::AnimationSpace
+						? GetTransformOrigin(InstanceHandle)
+						: FTransform::Identity;
 
 					const AActor* Actor = Cast<AActor>(BoundObject);
 					if (!Actor)
diff --git a/Engine/Plugins/Experimental/MovieScenePoseSearchTracks/Source/MovieScenePoseSearchTracks/Public/Sections/MovieSceneStitchAnimSection.h b/Engine/Plugins/Experimental/MovieScenePoseSearchTracks/Source/MovieScenePoseSearchTracks/Public/Sections/MovieSceneStitchAnimSection.h
--- a/Engine/Plugins/Experimental/MovieScenePoseSearchTracks/Source/MovieScenePoseSearchTracks/Public/Sections/MovieSceneStitchAnimSection.h
+++ b/Engine/Plugins/Experimental/MovieScenePoseSearchTracks/Source/MovieScenePoseSearchTracks/Public/Sections/MovieSceneStitchAnimSection.h
@@ -38,6 +38,9 @@ struct FMovieSceneStitchAnimComponentData
 	UPROPERTY()
 	FFrameNumber EndFrame;
 
+	UPROPERTY()
+	EMovieSceneRootMotionSpace TargetTransformSpace = EMovieSceneRootMotionSpace::AnimationSpace;
+
 	double MapTimeToSectionSeconds(FFrameTime InPosition, FFrameRate InFrameRate) const;
 };
 
@@ -132,6 +135,9 @@ public:
 	UPROPERTY()
 	FMovieSceneFloatChannel Weight;
 
+	UPROPERTY(EditAnywhere, Category = "Stitch")
+	EMovieSceneRootMotionSpace TargetTransformSpace = EMovieSceneRootMotionSpace::AnimationSpace;
+
 	//~ UMovieSceneSection interface
 	virtual EMovieSceneChannelProxyType CacheChannelProxy() override;
 	virtual UObject* GetSourceObject() const override;
