Sequencer - pose history override for stitch track, and warping alignment improvements

diff --git a/Engine/Plugins/Animation/PoseSearch/Source/Runtime/Private/PoseSearchFeatureChannel_Pose.cpp b/Engine/Plugins/Animation/PoseSearch/Source/Runtime/Private/PoseSearchFeatureChannel_Pose.cpp
--- a/Engine/Plugins/Animation/PoseSearch/Source/Runtime/Private/PoseSearchFeatureChannel_Pose.cpp
+++ b/Engine/Plugins/Animation/PoseSearch/Source/Runtime/Private/PoseSearchFeatureChannel_Pose.cpp
@@ -51,6 +51,7 @@ bool UPoseSearchFeatureChannel_Pose::Finalize(UPoseSearchSchema* Schema)
 #endif // WITH_EDITORONLY_DATA
 			Position->SampleTimeOffset = 0.f;
 			Position->InputQueryPose = InputQueryPose;
+			Position->PermutationTimeType = PermutationTimeType;
 			SubChannels.Add(Position);
 		}
 
@@ -68,6 +69,7 @@ bool UPoseSearchFeatureChannel_Pose::Finalize(UPoseSearchSchema* Schema)
 			HeadingX->SampleTimeOffset = 0.f;
 			HeadingX->HeadingAxis = EHeadingAxis::X;
 			HeadingX->InputQueryPose = InputQueryPose;
+			HeadingX->PermutationTimeType = PermutationTimeType;
 			SubChannels.Add(HeadingX);
 
 			UPoseSearchFeatureChannel_Heading* HeadingY = NewObject<UPoseSearchFeatureChannel_Heading>(this, NAME_None, RF_Transient);
@@ -82,6 +84,7 @@ bool UPoseSearchFeatureChannel_Pose::Finalize(UPoseSearchSchema* Schema)
 			HeadingY->SampleTimeOffset = 0.f;
 			HeadingY->HeadingAxis = EHeadingAxis::Y;
 			HeadingY->InputQueryPose = InputQueryPose;
+			HeadingY->PermutationTimeType = PermutationTimeType;
 			SubChannels.Add(HeadingY);
 		}
 
@@ -99,6 +102,7 @@ bool UPoseSearchFeatureChannel_Pose::Finalize(UPoseSearchSchema* Schema)
 			Velocity->SampleTimeOffset = 0.f;
 			Velocity->InputQueryPose = InputQueryPose;
 			Velocity->bUseCharacterSpaceVelocities = bUseCharacterSpaceVelocities;
+			Velocity->PermutationTimeType = PermutationTimeType;
 			SubChannels.Add(Velocity);
 		}
 
@@ -113,6 +117,12 @@ bool UPoseSearchFeatureChannel_Pose::Finalize(UPoseSearchSchema* Schema)
 			Phase->DebugColor = SampledBone.DebugColor;
 #endif // WITH_EDITORONLY_DATA
 			Phase->InputQueryPose = InputQueryPose;
+
+			if (PermutationTimeType != EPermutationTimeType::UseSampleTime)
+			{
+				UE_LOG(LogPoseSearch, Error, TEXT("UPoseSearchFeatureChannel_Pose::Finalize - unsupported PermutationTimeType %d for EPoseSearchBoneFlags::Phase in Schema '%s'"), (int32)PermutationTimeType, *Schema->GetName());
+				return false;
+			}
 			SubChannels.Add(Phase);
 		}
 	}
@@ -126,6 +136,16 @@ UE::PoseSearch::TLabelBuilder& UPoseSearchFeatureChannel_Pose::GetLabel(UE::Pose
 	GetOuterLabel(LabelBuilder, LabelFormat);
 	AppendLabelSeparator(LabelBuilder, UE::PoseSearch::ELabelFormat::Full_Horizontal);
 	LabelBuilder.Append(TEXT("Pose"));
+
+	if (PermutationTimeType == EPermutationTimeType::UsePermutationTime)
+	{
+		LabelBuilder.Append(TEXT("_PT"));
+	}
+	else if (PermutationTimeType == EPermutationTimeType::UseSampleToPermutationTime)
+	{
+		LabelBuilder.Append(TEXT("_SPT"));
+	}
+
 	if (DebugWeightGroupID != INDEX_NONE)
 	{
 		LabelBuilder.Appendf(TEXT("_%d"), DebugWeightGroupID);
diff --git a/Engine/Plugins/Animation/PoseSearch/Source/Runtime/Private/PoseSearchHistory.cpp b/Engine/Plugins/Animation/PoseSearch/Source/Runtime/Private/PoseSearchHistory.cpp
--- a/Engine/Plugins/Animation/PoseSearch/Source/Runtime/Private/PoseSearchHistory.cpp
+++ b/Engine/Plugins/Animation/PoseSearch/Source/Runtime/Private/PoseSearchHistory.cpp
@@ -1126,15 +1126,17 @@ bool FMemStackPoseHistory::GetTransformAtTime(float Time, FTransform& OutBoneTra
 
 		if (ReferenceBoneIndexType == WorldSpaceIndexType)
 		{
+			const FTransformTrajectory& Trajectory = GetTrajectory();
+
 			if (BoneIndexType == ComponentSpaceIndexType)
 			{
-				OutBoneTransform = PoseHistory->GetTrajectory().GetSampleAtTime(Time, bExtrapolate).GetTransform();
+				OutBoneTransform = Trajectory.GetSampleAtTime(Time, bExtrapolate).GetTransform();
 				return true;
 			}
 
 			// getting BoneIndexType in ComponentSpaceIndexType and then multiplying it with the component to world transform from the trajectory
 			const bool bSuccess = GetTransformAtTime(Time, OutBoneTransform, BoneIndexSkeleton, BoneIndexType, ComponentSpaceIndexType, bExtrapolate);
-			OutBoneTransform *= PoseHistory->GetTrajectory().GetSampleAtTime(Time, bExtrapolate).GetTransform();
+			OutBoneTransform *= Trajectory.GetSampleAtTime(Time, bExtrapolate).GetTransform();
 			return bSuccess;
 		}
 
@@ -1197,6 +1199,71 @@ void FMemStackPoseHistory::AddFuturePose(float Time, IComponentSpacePoseProvider
 	FutureEntries.InsertDefaulted_GetRef(LowerBoundIdx).Update(Time, ComponentSpacePoseProvider, GetBoneToTransformMap(), true, Curves, MakeConstArrayView(GetCollectedCurves()));
 }
 
+void FMemStackPoseHistory::ExtractAndAddFuturePoses(const UAnimationAsset* AnimationAsset, float AnimationTime, float FiniteDeltaTime, const FVector& BlendParame1ters, float IntervalTime, const USkeleton* OverrideSkeleton)
+{
+	FMemMark Mark(FMemStack::Get());
+
+	check(FiniteDeltaTime >= 0.f);
+	if (AnimationTime < FiniteDeltaTime)
+	{
+		UE_LOG(LogPoseSearch, Error, TEXT("FMemStackPoseHistory::ExtractAndAddFuturePose - provided AnimationTime (%f) is too small. Clamping it to minimum value of %f"), AnimationTime, FiniteDeltaTime);
+		AnimationTime = FiniteDeltaTime;
+	}
+
+	if (IntervalTime < FiniteDeltaTime)
+	{
+		UE_LOG(LogPoseSearch, Error, TEXT("FMemStackPoseHistory::ExtractAndAddFuturePose - provided IntervalTime (%f) is too small. Clamping it to minimum value of %f"), IntervalTime, FiniteDeltaTime);
+		AnimationTime = FiniteDeltaTime;
+	}
+
+	const USkeleton* Skeleton = OverrideSkeleton ? OverrideSkeleton : AnimationAsset->GetSkeleton();
+	TArray<uint16> BoneIndices;
+	const FBoneToTransformMap& BoneToTransformMap = GetBoneToTransformMap();
+	if (BoneToTransformMap.IsEmpty())
+	{
+		const int32 NumBones = Skeleton->GetReferenceSkeleton().GetNum();
+		BoneIndices.SetNum(NumBones);
+		for (int32 BoneIndex = 0; BoneIndex < NumBones; ++BoneIndex)
+		{
+			BoneIndices[BoneIndex] = BoneIndex;
+		}
+	}
+	else
+	{
+		for (const FBoneToTransformPair& BoneToTransformPair : BoneToTransformMap)
+		{
+			BoneIndices.Add(BoneToTransformPair.Key);
+		}
+
+		BoneIndices.Sort();
+		FAnimationRuntime::EnsureParentsPresent(BoneIndices, Skeleton->GetReferenceSkeleton());
+	}
+
+	FBoneContainer BoneContainer;
+	BoneContainer.InitializeTo(BoneIndices, UE::Anim::FCurveFilterSettings(UE::Anim::ECurveFilterMode::DisallowAll), *Skeleton);
+
+	// extracting 2 poses to be able to calculate velocities
+	FCSPose<FCompactPose> ComponentSpacePose;
+	FCompactPose Pose;
+	FBlendedCurve Curves;
+	Pose.SetBoneContainer(&BoneContainer);
+
+	// extracting 2 poses to be able to calculate velocities
+	const FAnimationAssetSampler Sampler(AnimationAsset, FTransform::Identity, FVector::ZeroVector, FAnimationAssetSampler::DefaultRootTransformSamplingRate, false, false);
+
+	const int32 NumOfPoseExtractions = FMath::IsNearlyZero(FiniteDeltaTime) ? 1 : 2;
+	for (int32 i = 0; i < NumOfPoseExtractions; ++i)
+	{
+		const float FuturePoseExtractionTime = AnimationTime + (i - 1) * FiniteDeltaTime;
+		const float FuturePoseAnimationTime = IntervalTime + (i - 1) * FiniteDeltaTime;
+
+		Sampler.ExtractPose(FuturePoseExtractionTime, Pose, Curves);
+		ComponentSpacePose.InitPose(Pose);
+		FComponentSpacePoseProvider ComponentSpacePoseProvider(ComponentSpacePose);
+		AddFuturePose(FuturePoseAnimationTime, ComponentSpacePoseProvider);
+	}
+}
+
 void FMemStackPoseHistory::AddFuturePose(float Time, FCSPose<FCompactPose>& ComponentSpacePose, const FBlendedCurve& Curves)
 {
 	FComponentSpacePoseProvider ComponentSpacePoseProvider(ComponentSpacePose);
diff --git a/Engine/Plugins/Animation/PoseSearch/Source/Runtime/Private/PoseSearchLibrary.cpp b/Engine/Plugins/Animation/PoseSearch/Source/Runtime/Private/PoseSearchLibrary.cpp
--- a/Engine/Plugins/Animation/PoseSearch/Source/Runtime/Private/PoseSearchLibrary.cpp
+++ b/Engine/Plugins/Animation/PoseSearch/Source/Runtime/Private/PoseSearchLibrary.cpp
@@ -435,7 +435,7 @@ void UPoseSearchLibrary::TraceMotionMatching(
 	float RecordingTime = 0.f;
 	if (!SearchContext.GetContexts().IsEmpty())
 	{
-		if (UObject* FirstObject = SearchContext.GetContexts()[0]->GetFirstObjectParam())
+		if (const UObject* FirstObject = SearchContext.GetContexts()[0]->GetFirstObjectParam())
 		{
 			RecordingTime = FObjectTrace::GetWorldElapsedTime(FirstObject->GetWorld());
 		}
@@ -452,15 +452,15 @@ void UPoseSearchLibrary::TraceMotionMatching(
 
 	for (int32 AnimInstanceIndex = 0; AnimInstanceIndex < AnimContextsNum; ++AnimInstanceIndex)
 	{
-		if (FChooserEvaluationContext* AnimContext = SearchContext.GetContexts()[AnimInstanceIndex])
+		if (const FChooserEvaluationContext* AnimContext = SearchContext.GetContexts()[AnimInstanceIndex])
 		{
-			UObject* FirstObject = AnimContext->GetFirstObjectParam();
+			const UObject* FirstObject = AnimContext->GetFirstObjectParam();
 			const UObject* SkeletalMeshComponent = nullptr;
-			if (const UAnimInstance* AnimInstance = Cast<UAnimInstance>(AnimContext->GetFirstObjectParam()))
+			if (const UAnimInstance* AnimInstance = Cast<UAnimInstance>(FirstObject))
 			{
 				SkeletalMeshComponent = AnimInstance->GetOuter();
 			}
-			else if (const UActorComponent* ActorComponent = Cast<UActorComponent>(AnimContext->GetFirstObjectParam()))
+			else if (const UActorComponent* ActorComponent = Cast<UActorComponent>(FirstObject))
 			{
 				const AActor* Actor = ActorComponent->GetOwner();
 				check(Actor);
@@ -479,8 +479,8 @@ void UPoseSearchLibrary::TraceMotionMatching(
 
 	for (int32 AnimInstanceIndex = 0; AnimInstanceIndex < AnimContextsNum; ++AnimInstanceIndex)
 	{
-		FChooserEvaluationContext* Context = SearchContext.GetContexts()[AnimInstanceIndex];
-		if (UObject* Object = Context->GetFirstObjectParam())
+		const FChooserEvaluationContext* Context = SearchContext.GetContexts()[AnimInstanceIndex];
+		if (const UObject* Object = Context->GetFirstObjectParam())
 		{
 			TRACE_OBJECT(Object);
 			SearchId = HashCombineFast(SearchId, GetTypeHash(FObjectTrace::GetObjectId(Object)));
@@ -614,7 +614,7 @@ void UPoseSearchLibrary::TraceMotionMatching(
 	TraceState.Cycle = FPlatformTime::Cycles64();
 
 	// @todo: avoid publishing duplicated TraceState in ALL the AnimContexts! -currently necessary for multi character-
-	for (FChooserEvaluationContext* Context : SearchContext.GetContexts())
+	for (const FChooserEvaluationContext* Context : SearchContext.GetContexts())
 	{
 		TraceState.AnimInstanceId = FObjectTrace::GetObjectId(Context->GetFirstObjectParam());
 		TraceState.NodeId = SearchId;
@@ -1043,9 +1043,9 @@ UE::PoseSearch::FSearchResult UPoseSearchLibrary::MotionMatch(
 	const FPoseSearchEvent& EventToSearch)
 {
 	TArray<FChooserEvaluationContext, TInlineAllocator<UE::PoseSearch::PreallocatedRolesNum>> Contexts;
-	int NumContexts = AnimContexts.Num(); 
+	const int NumContexts = AnimContexts.Num(); 
 	Contexts.SetNum(NumContexts);
-	for(int i = 0; i <NumContexts; i++)
+	for(int i = 0; i < NumContexts; i++)
 	{
 		Contexts[i].AddObjectParam(const_cast<UObject*>(AnimContexts[i]));
 	}
@@ -1054,7 +1054,27 @@ UE::PoseSearch::FSearchResult UPoseSearchLibrary::MotionMatch(
 }
 
 UE::PoseSearch::FSearchResult UPoseSearchLibrary::MotionMatch(
-	const TArrayView<FChooserEvaluationContext> AnimContexts,
+	const TArrayView<const UObject*> AnimContexts,
+	const TArrayView<const UE::PoseSearch::FRole> Roles,
+	const TArrayView<const UE::PoseSearch::IPoseHistory*> PoseHistories, 
+	const TArrayView<const UObject*> AssetsToSearch,
+	const FPoseSearchContinuingProperties& ContinuingProperties,
+	float DesiredPermutationTimeOffset,
+	const FPoseSearchEvent& EventToSearch)
+{
+	TArray<FChooserEvaluationContext, TInlineAllocator<UE::PoseSearch::PreallocatedRolesNum>> Contexts;
+	const int NumContexts = AnimContexts.Num(); 
+	Contexts.SetNum(NumContexts);
+	for(int i = 0; i < NumContexts; i++)
+	{
+		Contexts[i].AddObjectParam(const_cast<UObject*>(AnimContexts[i]));
+	}
+
+	return MotionMatch(Contexts, Roles, PoseHistories, AssetsToSearch, ContinuingProperties, DesiredPermutationTimeOffset, EventToSearch);
+}
+
+UE::PoseSearch::FSearchResult UPoseSearchLibrary::MotionMatch(
+	const TArrayView<FChooserEvaluationContext> Contexts,
 	const TArrayView<const UE::PoseSearch::FRole> Roles,
 	const TArrayView<const UE::PoseSearch::IPoseHistory*> PoseHistories, 
 	const TArrayView<const UObject*> AssetsToSearch,
@@ -1062,12 +1082,10 @@ UE::PoseSearch::FSearchResult UPoseSearchLibrary::MotionMatch(
 	const FPoseSearchFutureProperties& Future,
 	const FPoseSearchEvent& EventToSearch)
 {
-	check(!AnimContexts.IsEmpty() && AnimContexts.Num() == Roles.Num() && AnimContexts.Num() == PoseHistories.Num());
+	check(!Contexts.IsEmpty() && Contexts.Num() == Roles.Num() && Contexts.Num() == PoseHistories.Num());
 
 	using namespace UE::PoseSearch;
 
-	FSearchResult SearchResult;
-
 	FMemMark Mark(FMemStack::Get());
 
 	TArray<const IPoseHistory*, TInlineAllocator<PreallocatedRolesNum, TMemStackAllocator<>>> InternalPoseHistories;
@@ -1098,41 +1116,8 @@ UE::PoseSearch::FSearchResult UPoseSearchLibrary::MotionMatch(
 		{
 			if (const IPoseHistory* PoseHistory = InternalPoseHistories[RoleIndex])
 			{
-				if (USkeleton* Skeleton = const_cast<USkeleton*>(GetContextSkeleton(AnimContexts[RoleIndex])))
+				if (const USkeleton* Skeleton = GetContextSkeleton(Contexts[RoleIndex]))
 				{
-					MemStackPoseHistories[RoleIndex].Init(InternalPoseHistories[RoleIndex]);
-
-					TArray<uint16> BoneIndices;
-					const FBoneToTransformMap& BoneToTransformMap = PoseHistory->GetBoneToTransformMap();
-					if (BoneToTransformMap.IsEmpty())
-					{
-						const int32 NumBones = Skeleton->GetReferenceSkeleton().GetNum();
-						BoneIndices.SetNum(NumBones);
-						for (int32 BoneIndex = 0; BoneIndex < NumBones; ++BoneIndex)
-						{
-							BoneIndices[BoneIndex] = BoneIndex;
-						}
-					}
-					else
-					{
-						for (const FBoneToTransformPair& BoneToTransformPair : BoneToTransformMap)
-						{
-							BoneIndices.Add(BoneToTransformPair.Key);
-						}
-
-						BoneIndices.Sort();
-						FAnimationRuntime::EnsureParentsPresent(BoneIndices, Skeleton->GetReferenceSkeleton());
-					}
-
-					FBoneContainer BoneContainer;
-					BoneContainer.InitializeTo(BoneIndices, UE::Anim::FCurveFilterSettings(UE::Anim::ECurveFilterMode::DisallowAll), *Skeleton);
-
-					// extracting 2 poses to be able to calculate velocities
-					FCSPose<FCompactPose> ComponentSpacePose;
-					FCompactPose Pose;
-					FBlendedCurve Curves;
-					Pose.SetBoneContainer(&BoneContainer);
-
 					// @todo: add input BlendParameters to support sampling FutureAnimation blendspaces and support for multi character
 					const UAnimationAsset* AnimationAsset = Cast<UAnimationAsset>(Future.Animation);
 					if (!AnimationAsset)
@@ -1147,40 +1132,46 @@ UE::PoseSearch::FSearchResult UPoseSearchLibrary::MotionMatch(
 						}
 					}
 
-					const FAnimationAssetSampler Sampler(AnimationAsset, FTransform::Identity, FVector::ZeroVector, FAnimationAssetSampler::DefaultRootTransformSamplingRate, false, false);
-					for (int32 i = 0; i < 2; ++i)
-					{
-						const float FuturePoseExtractionTime = FutureAnimationTime + (i - 1) * FiniteDelta;
-						const float FuturePoseAnimationTime = FutureIntervalTime + (i - 1) * FiniteDelta;
-
-						Sampler.ExtractPose(FuturePoseExtractionTime, Pose, Curves);
-						ComponentSpacePose.InitPose(Pose);
-						FComponentSpacePoseProvider ComponentSpacePoseProvider(ComponentSpacePose);
-						MemStackPoseHistories[RoleIndex].AddFuturePose(FuturePoseAnimationTime, ComponentSpacePoseProvider);
-					}
-
+					MemStackPoseHistories[RoleIndex].Init(InternalPoseHistories[RoleIndex]);
+					MemStackPoseHistories[RoleIndex].ExtractAndAddFuturePoses(AnimationAsset, FutureAnimationTime, FiniteDelta, FVector::ZeroVector, FutureIntervalTime, Skeleton);
 					InternalPoseHistories[RoleIndex] = MemStackPoseHistories[RoleIndex].GetThisOrPoseHistory();
 				}
 			}
-		}		
-	}
+		}
+	}		
+
+	return MotionMatch(Contexts, Roles, InternalPoseHistories, AssetsToSearch, ContinuingProperties, FutureIntervalTime, EventToSearch);
+}
+
+UE::PoseSearch::FSearchResult UPoseSearchLibrary::MotionMatch(
+	const TArrayView<FChooserEvaluationContext> Contexts,
+	const TArrayView<const UE::PoseSearch::FRole> Roles,
+	const TArrayView<const UE::PoseSearch::IPoseHistory*> PoseHistories, 
+	const TArrayView<const UObject*> AssetsToSearch,
+	const FPoseSearchContinuingProperties& ContinuingProperties,
+	const float DesiredPermutationTimeOffset,
+	const FPoseSearchEvent& EventToSearch)
+{
+	using namespace UE::PoseSearch;
+
+	FSearchResult SearchResult;
 
 	FReconstructedPreviousSearchBestResultMap ReconstructedPreviousSearchBestResultMap;
 	FSearchResult ReconstructedPreviousSearchResult;
-	FSearchContext SearchContext(FutureIntervalTime, nullptr, ReconstructedPreviousSearchResult, FFloatInterval(0.f, 0.f), EventToSearch);
+	FSearchContext SearchContext(DesiredPermutationTimeOffset, nullptr, ReconstructedPreviousSearchResult, FFloatInterval(0.f, 0.f), EventToSearch);
 
 	for (int32 RoleIndex = 0; RoleIndex < Roles.Num(); ++RoleIndex)
 	{
-		SearchContext.AddRole(Roles[RoleIndex], &AnimContexts[RoleIndex], InternalPoseHistories[RoleIndex]);
+		SearchContext.AddRole(Roles[RoleIndex], &Contexts[RoleIndex], PoseHistories[RoleIndex]);
 	}
 	
 	bool bAsyncBuildIndexInProgress = false;
 
 	// collecting all the possible continuing pose search (it could be multiple searches, but most likely only one)
 	float DeltaSeconds = FiniteDelta;
-	if (AnimContexts[0].ObjectParams.Num() > 0)
+	if (Contexts[0].ObjectParams.Num() > 0)
 	{
-		if (const UAnimInstance* AnimInstance = Cast<UAnimInstance>(AnimContexts[0].GetFirstObjectParam()))
+		if (const UAnimInstance* AnimInstance = Cast<UAnimInstance>(Contexts[0].GetFirstObjectParam()))
 		{
 			DeltaSeconds = AnimInstance->GetDeltaSeconds();
 		}
diff --git a/Engine/Plugins/Animation/PoseSearch/Source/Runtime/Public/PoseSearch/PoseSearchFeatureChannel_Pose.h b/Engine/Plugins/Animation/PoseSearch/Source/Runtime/Public/PoseSearch/PoseSearchFeatureChannel_Pose.h
--- a/Engine/Plugins/Animation/PoseSearch/Source/Runtime/Public/PoseSearch/PoseSearchFeatureChannel_Pose.h
+++ b/Engine/Plugins/Animation/PoseSearch/Source/Runtime/Public/PoseSearch/PoseSearchFeatureChannel_Pose.h
@@ -70,6 +70,9 @@ public:
 	UPROPERTY(EditAnywhere, Category = "Settings")
 	bool bUseCharacterSpaceVelocities = true;
 
+	UPROPERTY(EditAnywhere, Category = "Settings")
+	EPermutationTimeType PermutationTimeType = EPermutationTimeType::UseSampleTime;
+
 	UPROPERTY(Transient)
 	TArray<TObjectPtr<UPoseSearchFeatureChannel>> SubChannels;
 
diff --git a/Engine/Plugins/Animation/PoseSearch/Source/Runtime/Public/PoseSearch/PoseSearchHistory.h b/Engine/Plugins/Animation/PoseSearch/Source/Runtime/Public/PoseSearch/PoseSearchHistory.h
--- a/Engine/Plugins/Animation/PoseSearch/Source/Runtime/Public/PoseSearch/PoseSearchHistory.h
+++ b/Engine/Plugins/Animation/PoseSearch/Source/Runtime/Public/PoseSearch/PoseSearchHistory.h
@@ -281,6 +281,12 @@ struct POSESEARCH_API FMemStackPoseHistory : public IPoseHistory
 
 	void AddFutureRootBone(float Time, const FTransform& FutureRootBoneTransform, bool bStoreScales);
 	void AddFuturePose(float Time, IComponentSpacePoseProvider& ComponentSpacePoseProvider, const FBlendedCurve& Curves = FBlendedCurve());
+	
+	// Experimental, this feature might be removed without warning, not for production use
+	// extracting a pose from AnimationAsset at AnimationTime and  (using BlendParameters in case it's a blend space) and placing it in the FMemStackPoseHistory at IntervalTime seconds
+	// if FiniteDeltaTime is > FiniteDeltaTime a second pose will be extracted at AnimationTime - FiniteDeltaTime and placed at IntervalTime - FiniteDeltaTime (for MM to be able to calculate velocities)
+	void ExtractAndAddFuturePoses(const UAnimationAsset* AnimationAsset, float AnimationTime, float FiniteDeltaTime, const FVector& BlendParameters, float IntervalTime, const USkeleton* OverrideSkeleton = nullptr);
+
 	UE_DEPRECATED(5.6, "Use other AddFuturePose signatures instead")
 	void AddFuturePose(float Time, FCSPose<FCompactPose>& ComponentSpacePose, const FBlendedCurve& Curves = FBlendedCurve());
 
diff --git a/Engine/Plugins/Animation/PoseSearch/Source/Runtime/Public/PoseSearch/PoseSearchLibrary.h b/Engine/Plugins/Animation/PoseSearch/Source/Runtime/Public/PoseSearch/PoseSearchLibrary.h
--- a/Engine/Plugins/Animation/PoseSearch/Source/Runtime/Public/PoseSearch/PoseSearchLibrary.h
+++ b/Engine/Plugins/Animation/PoseSearch/Source/Runtime/Public/PoseSearch/PoseSearchLibrary.h
@@ -287,16 +287,36 @@ public:
 		const FPoseSearchFutureProperties& Future,
 		const FPoseSearchEvent& EventToSearch);
 	
+	// Experimental, this feature might be removed without warning, not for production use
 	static UE::PoseSearch::FSearchResult MotionMatch(
-			const TArrayView<FChooserEvaluationContext> Contexts,
-    		const TArrayView<const UE::PoseSearch::FRole> Roles,
-    		const TArrayView<const UE::PoseSearch::IPoseHistory*> PoseHistories, 
-    		const TArrayView<const UObject*> AssetsToSearch,
-    		const FPoseSearchContinuingProperties& ContinuingProperties,
-    		const FPoseSearchFutureProperties& Future,
-    		const FPoseSearchEvent& EventToSearch);
-	
-		
+		const TArrayView<FChooserEvaluationContext> Contexts,
+		const TArrayView<const UE::PoseSearch::FRole> Roles,
+		const TArrayView<const UE::PoseSearch::IPoseHistory*> PoseHistories, 
+		const TArrayView<const UObject*> AssetsToSearch,
+		const FPoseSearchContinuingProperties& ContinuingProperties,
+		const FPoseSearchFutureProperties& Future,
+		const FPoseSearchEvent& EventToSearch);
+			
+	// Experimental, this feature might be removed without warning, not for production use
+	static UE::PoseSearch::FSearchResult MotionMatch(
+		const TArrayView<FChooserEvaluationContext> Contexts,
+		const TArrayView<const UE::PoseSearch::FRole> Roles,
+		const TArrayView<const UE::PoseSearch::IPoseHistory*> PoseHistories,
+		const TArrayView<const UObject*> AssetsToSearch,
+		const FPoseSearchContinuingProperties& ContinuingProperties,
+		const float DesiredPermutationTimeOffset,
+		const FPoseSearchEvent& EventToSearch);
+
+	// Experimental, this feature might be removed without warning, not for production use
+	static UE::PoseSearch::FSearchResult MotionMatch(
+		const TArrayView<const UObject*> AnimContexts,
+		const TArrayView<const UE::PoseSearch::FRole> Roles,
+		const TArrayView<const UE::PoseSearch::IPoseHistory*> PoseHistories,
+		const TArrayView<const UObject*> AssetsToSearch,
+		const FPoseSearchContinuingProperties& ContinuingProperties,
+		const float DesiredPermutationTimeOffset,
+		const FPoseSearchEvent& EventToSearch);
+
 	UE_DEPRECATED(5.5, "Use other MotionMatch signatures instead")
 	static UE::PoseSearch::FSearchResult MotionMatch(
 		const TArrayView<UAnimInstance*> AnimInstances,
diff --git a/Engine/Plugins/Experimental/MovieScenePoseSearchTracks/Source/MovieScenePoseSearchTracks/Private/Sections/MovieSceneStitchAnimSection.cpp b/Engine/Plugins/Experimental/MovieScenePoseSearchTracks/Source/MovieScenePoseSearchTracks/Private/Sections/MovieSceneStitchAnimSection.cpp
--- a/Engine/Plugins/Experimental/MovieScenePoseSearchTracks/Source/MovieScenePoseSearchTracks/Private/Sections/MovieSceneStitchAnimSection.cpp
+++ b/Engine/Plugins/Experimental/MovieScenePoseSearchTracks/Source/MovieScenePoseSearchTracks/Private/Sections/MovieSceneStitchAnimSection.cpp
@@ -1,32 +1,54 @@
 // Copyright Epic Games, Inc. All Rights Reserved.
 
 #include "Sections/MovieSceneStitchAnimSection.h"
-#include "Channels/MovieSceneChannelProxy.h"
-#include "Tracks/MovieSceneStitchAnimTrack.h"
-#include "PoseSearchTracksComponentTypes.h"
-#include "PoseSearch/PoseSearchHistory.h"
+#include "Animation/AnimInstance.h"
 #include "Animation/AnimPoseSearchProvider.h"
-#include "EntitySystem/BuiltInComponentTypes.h"
-#include "Systems/MovieSceneSkeletalAnimationSystem.h"
-#include "MovieSceneTracksComponentTypes.h"
+#include "Animation/AnimSequence.h"
+#include "Animation/BuiltInAttributeTypes.h"
 #include "AnimMixerComponentTypes.h"
-#include "PoseSearch/PoseSearchHistoryAttribute.h"
-#include "EvaluationVM/KeyframeState.h"
+#include "Channels/MovieSceneChannelProxy.h"
+#include "DecompressionTools.h"
+#include "EntitySystem/BuiltInComponentTypes.h"
 #include "EvaluationVM/EvaluationVM.h"
-#include "PoseSearch/PoseSearchLibrary.h"
+#include "EvaluationVM/KeyframeState.h"
+#include "MovieSceneTracksComponentTypes.h"
+#include "PoseSearch/PoseSearchContext.h"
 #include "PoseSearch/PoseSearchDatabase.h"
 #include "PoseSearch/PoseSearchFeatureChannel_PermutationTime.h"
+#include "PoseSearch/PoseSearchHistory.h"
+#include "PoseSearch/PoseSearchHistoryAttribute.h"
+#include "PoseSearch/PoseSearchLibrary.h"
 #include "PoseSearch/PoseSearchSchema.h"
-#include "DecompressionTools.h"
-#include "Animation/AnimInstance.h"
-#include "Animation/AnimSequence.h"
+#include "PoseSearchTracksComponentTypes.h"
+#include "Systems/MovieSceneSkeletalAnimationSystem.h"
+#include "Tracks/MovieSceneStitchAnimTrack.h"
 #include "Variants/MovieSceneTimeWarpVariantPayloads.h"
-#include "Animation/BuiltInAttributeTypes.h"
+#include "VisualLogger/VisualLogger.h"
 
 #include UE_INLINE_GENERATED_CPP_BY_NAME(MovieSceneStitchAnimSection)
 
 #define LOCTEXT_NAMESPACE "MovieSceneStitchAnimSection"
 
+namespace UE::MovieScene
+{
+	// @todo: we could inherit from MemStackPoseHistory and just calcualte future poses here rather than relying on FPoseSearchFutureProperties
+	// Experimental, this feature might be removed without warning, not for production use
+	struct FOverridePoseHistory : UE::PoseSearch::FMemStackPoseHistory
+	{
+		virtual const FTransformTrajectory& GetTrajectory() const override
+		{
+			return Trajectory;
+		}
+
+		virtual void SetTrajectory(const FTransformTrajectory& InTrajectory, float InTrajectorySpeedMultiplier = 1.f) override
+		{
+			Trajectory = InTrajectory;
+		}
+
+		FTransformTrajectory Trajectory;
+	};
+} // namespace UE::MovieScene
+
 void FMovieSceneStitchAnimEvaluationTask::Execute(UE::AnimNext::FEvaluationVM& VM) const
 {
 	using namespace UE::AnimNext;
@@ -59,19 +81,49 @@ void FMovieSceneStitchAnimEvaluationTask::Execute(UE::AnimNext::FEvaluationVM& V
 				PlayingAsset.Asset = GetAnimAsset();
 				PlayingAsset.AccumulatedTime = GetAccumulatedTime();*/
 
+			// set up the override pose history we're using for the UPoseSearchLibrary::MotionMatch. We initiale it with PoseHistory as fallback for non overridden methods
+			FOverridePoseHistory OverridePoseHistory;
+			OverridePoseHistory.Init(PoseHistory);
+
+			// computing the overridden trajectory with all the past FTransformTrajectorySample from the PoseHistory we found via FPoseHistoryAnimationAttribute, 
+			// and one prediction future sample computed from StitchData.TargetTransform and TimeToTarget
+			const TArray<FTransformTrajectorySample>& TrajectorySamples = PoseHistory->GetTrajectory().Samples;
+			OverridePoseHistory.Trajectory.Samples.Reserve(TrajectorySamples.Num()); // over estimating memory requirements to avoid reallocations
+			for (const FTransformTrajectorySample& TrajectorySample : TrajectorySamples)
+			{
+				// we want to only collect the past trajectory, not the prediction for now
+				// @todo: calculate the full trajectory in Sequencer!
+				if (TrajectorySample.TimeInSeconds > 0.f)
+				{
+					break;
+				}
+
+				OverridePoseHistory.Trajectory.Samples.Add(TrajectorySample);
+			}
+			// @todo: FutureTrajectorySample should contain the position / facing or the mesh in world space,
+			//        not the world space transform StitchData.TargetTransform that seems to be representing the actor transform!
+			FTransformTrajectorySample FutureTrajectorySample;
+
+			const FTransform MeshTargetTransform = MeshToActorTransform * StitchData.TargetTransform;
+			FutureTrajectorySample.Position = MeshTargetTransform.GetLocation();
+			FutureTrajectorySample.Facing = MeshTargetTransform.GetRotation();
+			FutureTrajectorySample.TimeInSeconds = TimeToTarget;
+			OverridePoseHistory.Trajectory.Samples.Add(FutureTrajectorySample);
+
+			// extracting two future poses from StitchData.TargetPoseAsset at StitchData.TargetAnimationTimeSeconds in the future and StitchData.TargetAnimationTimeSeconds - FiniteDelta
+			// to have motion matching been able to compute velocities (if required by the database schema channels)
+			OverridePoseHistory.ExtractAndAddFuturePoses(StitchData.TargetPoseAsset, StitchData.TargetAnimationTimeSeconds, FiniteDelta, FVector::ZeroVector, TimeToTarget);
+
+			PoseHistory = &OverridePoseHistory;
+
 			FPoseSearchContinuingProperties ContinuingProperties;
 			/*ContinuingProperties.PlayingAsset = PlayingAsset.Asset;
 			ContinuingProperties.PlayingAssetAccumulatedTime = PlayingAsset.AccumulatedTime;*/
 
-			FPoseSearchFutureProperties Future;
-			Future.Animation = StitchData.TargetPoseAsset;
-			Future.AnimationTime = StitchData.TargetAnimationTimeSeconds;
-			Future.IntervalTime = TimeToTarget;
 			const UObject* StitchDatabaseObj = StitchData.StitchDatabase.Get();
-
 			const UObject* AnimContext = ContextObject.Get();
 			const FSearchResult SearchResult = UPoseSearchLibrary::MotionMatch(MakeArrayView(&AnimContext, 1),
-				MakeArrayView(&DefaultRole, 1), MakeArrayView(&PoseHistory, 1), MakeArrayView(&StitchDatabaseObj, 1), ContinuingProperties, Future, FPoseSearchEvent());
+				MakeArrayView(&DefaultRole, 1), MakeArrayView(&PoseHistory, 1), MakeArrayView(&StitchDatabaseObj, 1), ContinuingProperties, TimeToTarget, FPoseSearchEvent());
 
 			if (const FSearchIndexAsset* SearchIndexAsset = SearchResult.GetSearchIndexAsset())
 			{
@@ -86,7 +138,7 @@ void FMovieSceneStitchAnimEvaluationTask::Execute(UE::AnimNext::FEvaluationVM& V
 
 					// figuring out the WantedPlayRate
 					MatchedAssetPlayRate = 1.f;
-					if (Future.Animation && Future.IntervalTime > 0.f)
+					if (StitchData.TargetPoseAsset && TimeToTarget > 0.f)
 					{
 						if (const UPoseSearchFeatureChannel_PermutationTime* PermutationTimeChannel = Database->Schema->FindFirstChannelOfType<UPoseSearchFeatureChannel_PermutationTime>())
 						{
@@ -95,7 +147,7 @@ void FMovieSceneStitchAnimEvaluationTask::Execute(UE::AnimNext::FEvaluationVM& V
 							{
 								TConstArrayView<float> ResultData = Database->GetSearchIndex().GetPoseValues(SearchResult.PoseIdx);
 								const float ActualIntervalTime = PermutationTimeChannel->GetPermutationTime(ResultData);
-								MatchedAssetPlayRate = ActualIntervalTime / Future.IntervalTime;
+								MatchedAssetPlayRate = ActualIntervalTime / TimeToTarget;
 							}
 						}
 					}
@@ -123,52 +175,46 @@ void FMovieSceneStitchAnimEvaluationTask::Execute(UE::AnimNext::FEvaluationVM& V
 		VM.PushValue(KEYFRAME_STACK_NAME, MakeUnique<FKeyframeState>(MoveTemp(ReferenceKeyframe)));
 		return;
 	}
-	const bool bUseRawData = FDecompressionTools::ShouldUseRawData(AnimSequence, Keyframe.Pose);
-
-	float AnimationTime = MatchedAssetTime + ((CurrentTime - InitialTime) * MatchedAssetPlayRate);
-	float PreviousAnimationTime = AnimationTime - ((CurrentTime - PreviousTime) * MatchedAssetPlayRate);
-
-	// Calculate the root transform. 
-	// We do this by aligning and projecting the matched animation asset twice,
-	// first to the 'initial' transform, and then to the 'target' transform, and blending the results based on time proximity.
-
-	struct
-	{
-		FTransform Initial;
-		FTransform Current;
-		FTransform Target;
-	} AnimSpace;
-
-	AnimSpace.Initial = AnimSequence->ExtractRootTrackTransform(FAnimExtractContext((double)MatchedAssetTime), nullptr);
-	AnimSpace.Current = AnimSequence->ExtractRootTrackTransform(FAnimExtractContext((double)AnimationTime), nullptr);
-	AnimSpace.Target = AnimSequence->ExtractRootTrackTransform(FAnimExtractContext((double)(AnimationTime + (TimeToTarget * MatchedAssetPlayRate))), nullptr);
-
-	FTransform TargetTransformWorldSpace = StitchData.TargetTransform;
-	if (InverseMeshToActorRotation)
-	{
-		AnimSpace.Initial.SetTranslation(InverseMeshToActorRotation.GetValue() * AnimSpace.Initial.GetTranslation());
-		AnimSpace.Current.SetTranslation(InverseMeshToActorRotation.GetValue() * AnimSpace.Current.GetTranslation());
-		AnimSpace.Target.SetTranslation(InverseMeshToActorRotation.GetValue() * AnimSpace.Target.GetTranslation());
-		TargetTransformWorldSpace.SetRotation(InverseMeshToActorRotation.GetValue() * TargetTransformWorldSpace.GetRotation());
-		TargetTransformWorldSpace.NormalizeRotation();
-	}
 
-	TargetTransformWorldSpace *= SequenceTransformOrigin;
-
-	FTransform WorldSpaceProjectedFromStart = AnimSpace.Current.GetRelativeTransform(AnimSpace.Initial) * InitialRootTransform;
-	FTransform WorldSpaceProjectedFromTarget = AnimSpace.Current.GetRelativeTransform(AnimSpace.Target) * TargetTransformWorldSpace;
-
-	FTransform RootTransformWorldSpace;
-	// TODO: Protect against division by 0
-	float BlendWeight = ((CurrentTime - InitialTime) / (TimeToTarget + CurrentTime - InitialTime));
-	RootTransformWorldSpace.Blend(WorldSpaceProjectedFromStart, WorldSpaceProjectedFromTarget, BlendWeight);
-
-	FQuat NewRotation = FQuat::Slerp(WorldSpaceProjectedFromStart.GetRotation(), WorldSpaceProjectedFromTarget.GetRotation(), BlendWeight);
-	RootTransformWorldSpace.SetRotation(NewRotation);
+	const float AnimationTime = MatchedAssetTime + ((CurrentTime - InitialTime) * MatchedAssetPlayRate);
+	const float PreviousAnimationTime = AnimationTime - ((CurrentTime - PreviousTime) * MatchedAssetPlayRate);
+
+	// Calculate the root transform by warping linearly the animation root transforms to the target StitchData.TargetTransform (after converting it to the correct animation space in MeshComponent space)
+	const float BlendWeightDenominator = TimeToTarget + CurrentTime - InitialTime;
+	const float BlendWeight = FMath::IsNearlyZero(BlendWeightDenominator) ? 1.f : (CurrentTime - InitialTime) / BlendWeightDenominator;
+
+	// mesh space animation transforms associate to the beginning of the stitch track AnimSpaceInitial, end AnimSpaceTarget, and current AnimSpaceCurrent
+	const FTransform AnimSpaceInitial = AnimSequence->ExtractRootTrackTransform(FAnimExtractContext((double)MatchedAssetTime), nullptr);
+	const FTransform AnimSpaceCurrent = AnimSequence->ExtractRootTrackTransform(FAnimExtractContext((double)AnimationTime), nullptr);
+	const FTransform AnimSpaceTarget = AnimSequence->ExtractRootTrackTransform(FAnimExtractContext((double)(AnimationTime + (TimeToTarget * MatchedAssetPlayRate))), nullptr);
+
+	// local delta transforms from the beginning of the animation to the end (InitialToTarget), and current playback time (InitialToCurrent)
+	const FTransform InitialToTarget = AnimSpaceTarget.GetRelativeTransform(AnimSpaceInitial);
+	const FTransform InitialToCurrent = AnimSpaceCurrent.GetRelativeTransform(AnimSpaceInitial);
+
+	// calculating transforms in mesh space for the beginning of the stitch..
+	const FTransform InitialMeshTransform = MeshToActorTransform * InitialRootTransform;
+	// ..final transform where we want the actor to be placed..
+	const FTransform FinalMeshTransform = MeshToActorTransform * StitchData.TargetTransform;
+	// ..where the animation will bring the actor in case no warping is applied..
+	const FTransform FinalAnimationTransform = InitialToTarget * InitialMeshTransform;
+	// ..where the animation would place the actor in case no warping is applied..
+	const FTransform CurrentAnimationTransform = InitialToCurrent * InitialMeshTransform;
+
+	// calculating the warping alignement error, as required delta transform from where the animation would end up, and where we want to place the actor
+	const FTransform AlignmentError = FinalMeshTransform.GetRelativeTransform(FinalAnimationTransform);
+	// blending the AlignmentError by BlendWeight
+	FTransform BlendedAlignmentError;
+	BlendedAlignmentError.Blend(FTransform::Identity, AlignmentError, BlendWeight);
+	// calculating where the mesh want to be with after applying the animation AND the blended portion from the alignement error (BlendedAlignmentError)
+	const FTransform CurrentMeshTransform = BlendedAlignmentError * CurrentAnimationTransform;
+	// converting CurrentMeshTransform to world actor space so we can set it to the FTransformAnimationAttribute
+	const FTransform CurrentRootTransform = MeshToActorTransform.GetRelativeTransformReverse(CurrentMeshTransform);
 
 	FDeltaTimeRecord DeltaTime;
 	DeltaTime.Set(PreviousAnimationTime, AnimationTime-PreviousAnimationTime);
 	const FAnimExtractContext ExtractionContext((double)AnimationTime, false, DeltaTime, false);
+	const bool bUseRawData = FDecompressionTools::ShouldUseRawData(AnimSequence, Keyframe.Pose);
 
 	if (EnumHasAnyFlags(VM.GetFlags(), EEvaluationFlags::Bones))
 	{
@@ -191,13 +237,38 @@ void FMovieSceneStitchAnimEvaluationTask::Execute(UE::AnimNext::FEvaluationVM& V
 	if (EnumHasAnyFlags(VM.GetFlags(), EEvaluationFlags::Attributes | EEvaluationFlags::Trajectory))
 	{
 		FTransformAnimationAttribute* RootMotionAttribute = Keyframe.Attributes.FindOrAdd<FTransformAnimationAttribute>(AnimMixerComponents->RootTransformAttributeId);
-		RootMotionAttribute->Value = RootTransformWorldSpace;
+		RootMotionAttribute->Value = CurrentRootTransform;
 
 		FFloatAnimationAttribute* RootMotionWeightAttribute = Keyframe.Attributes.FindOrAdd<FFloatAnimationAttribute>(AnimMixerComponents->RootTransformWeightAttributeId);
 		RootMotionWeightAttribute->Value = 1.f;
 	}
 
 	VM.PushValue(KEYFRAME_STACK_NAME, MakeUnique<FKeyframeState>(MoveTemp(Keyframe)));
+
+#if ENABLE_VISUAL_LOG
+	if (FVisualLogger::IsRecording())
+	{
+		static const TCHAR* LogName = TEXT("MovieSceneStitchAnimSection");
+
+		// actor space debug lines
+		UE_VLOG_CIRCLE(ContextObject.Get(), LogName, Display, InitialRootTransform.GetLocation(), FVector::UpVector, 10.f, FColorList::Blue, TEXT(""));
+		UE_VLOG_CIRCLE(ContextObject.Get(), LogName, Display, CurrentRootTransform.GetLocation(), FVector::UpVector, 10.f, FColorList::Blue, TEXT(""));
+		UE_VLOG_CIRCLE(ContextObject.Get(), LogName, Display, StitchData.TargetTransform.GetLocation(), FVector::UpVector, 10.f, FColorList::Blue, TEXT(""));
+		UE_VLOG_SEGMENT(ContextObject.Get(), LogName, Display, InitialRootTransform.GetLocation(), StitchData.TargetTransform.GetLocation(), FColorList::Blue, TEXT(""));
+		
+		// mesh space animation debug lines
+		UE_VLOG_CIRCLE(ContextObject.Get(), LogName, Display, InitialMeshTransform.GetLocation(), FVector::UpVector, 15.f, FColorList::Green, TEXT(""));
+		UE_VLOG_CIRCLE(ContextObject.Get(), LogName, Display, CurrentAnimationTransform.GetLocation(), FVector::UpVector, 15.f, FColorList::Green, TEXT(""));
+		UE_VLOG_CIRCLE(ContextObject.Get(), LogName, Display, FinalAnimationTransform.GetLocation(), FVector::UpVector, 15.f, FColorList::Green, TEXT(""));
+		UE_VLOG_SEGMENT(ContextObject.Get(), LogName, Display, InitialMeshTransform.GetLocation(), FinalAnimationTransform.GetLocation(), FColorList::Green, TEXT(""));
+		
+		// mesh space full aligned warped animation debug lines
+		UE_VLOG_CIRCLE(ContextObject.Get(), LogName, Display, InitialMeshTransform.GetLocation(), FVector::UpVector, 18.f, FColorList::Black, TEXT(""));
+		UE_VLOG_CIRCLE(ContextObject.Get(), LogName, Display, CurrentMeshTransform.GetLocation(), FVector::UpVector, 18.f, FColorList::Black, TEXT(""));
+		UE_VLOG_CIRCLE(ContextObject.Get(), LogName, Display, FinalMeshTransform.GetLocation(), FVector::UpVector, 18.f, FColorList::Black, TEXT(""));
+		UE_VLOG_SEGMENT(ContextObject.Get(), LogName, Display, InitialMeshTransform.GetLocation(), FinalMeshTransform.GetLocation(), FColorList::Black, TEXT(""));
+	}
+#endif
 }
 
 UMovieSceneStitchAnimSection::UMovieSceneStitchAnimSection(const FObjectInitializer& ObjectInitializer)
diff --git a/Engine/Plugins/Experimental/MovieScenePoseSearchTracks/Source/MovieScenePoseSearchTracks/Private/Systems/MovieSceneStitchAnimSystem.cpp b/Engine/Plugins/Experimental/MovieScenePoseSearchTracks/Source/MovieScenePoseSearchTracks/Private/Systems/MovieSceneStitchAnimSystem.cpp
--- a/Engine/Plugins/Experimental/MovieScenePoseSearchTracks/Source/MovieScenePoseSearchTracks/Private/Systems/MovieSceneStitchAnimSystem.cpp
+++ b/Engine/Plugins/Experimental/MovieScenePoseSearchTracks/Source/MovieScenePoseSearchTracks/Private/Systems/MovieSceneStitchAnimSystem.cpp
@@ -163,7 +163,7 @@ namespace UE::MovieScene
 
 						if (SkeletalMeshComponent)
 						{
-							AnimTask->InverseMeshToActorRotation = Actor->GetRootComponent()->GetComponentTransform().GetRelativeTransformReverse(SkeletalMeshComponent->GetComponentTransform()).GetRotation();
+							AnimTask->MeshToActorTransform = SkeletalMeshComponent->GetRelativeTransform();
 						}
 					}
 				}
diff --git a/Engine/Plugins/Experimental/MovieScenePoseSearchTracks/Source/MovieScenePoseSearchTracks/Public/Sections/MovieSceneStitchAnimSection.h b/Engine/Plugins/Experimental/MovieScenePoseSearchTracks/Source/MovieScenePoseSearchTracks/Public/Sections/MovieSceneStitchAnimSection.h
--- a/Engine/Plugins/Experimental/MovieScenePoseSearchTracks/Source/MovieScenePoseSearchTracks/Public/Sections/MovieSceneStitchAnimSection.h
+++ b/Engine/Plugins/Experimental/MovieScenePoseSearchTracks/Source/MovieScenePoseSearchTracks/Public/Sections/MovieSceneStitchAnimSection.h
@@ -54,46 +54,46 @@ struct FMovieSceneStitchAnimEvaluationTask : public FAnimNextEvaluationTask
 	// Task entry point
 	virtual void Execute(UE::AnimNext::FEvaluationVM& VM) const override;
 
-	UPROPERTY()
+	UPROPERTY(Transient)
 	FMovieSceneStitchAnimComponentData StitchData;
 
 	// Initial sequencer time when the motion matching calculation was done
-	UPROPERTY()
-	float InitialTime = -1.0f;
+	UPROPERTY(Transient)
+	float InitialTime = -1.f;
 
 	// Time until our target pose
-	UPROPERTY()
-	float TimeToTarget = -1.0f;
+	UPROPERTY(Transient)
+	float TimeToTarget = -1.f;
 
 	// Previous sequencer time
-	UPROPERTY()
-	float PreviousTime = -1.0f;
+	UPROPERTY(Transient)
+	float PreviousTime = -1.f;
 
 	// Current sequencer time
-	UPROPERTY()
-	float CurrentTime = -1.0f;
+	UPROPERTY(Transient)
+	float CurrentTime = -1.f;
 
-	UPROPERTY()
-	FTransform InitialRootTransform;
+	UPROPERTY(Transient)
+	FTransform InitialRootTransform = FTransform::Identity;
 
-	UPROPERTY()
+	UPROPERTY(Transient)
 	TWeakObjectPtr<const UObject> ContextObject;
 
-	UPROPERTY()
-	FTransform SequenceTransformOrigin;
+	UPROPERTY(Transient)
+	FTransform SequenceTransformOrigin = FTransform::Identity;
 
 	// Calculated and cached by the motion match at InitialTime
-	UPROPERTY()
+	UPROPERTY(Transient)
 	mutable TObjectPtr<UAnimationAsset> MatchedAsset;
 
-	UPROPERTY()
+	UPROPERTY(Transient)
 	mutable float MatchedAssetTime = 0.f;
 
-	UPROPERTY()
-	mutable float MatchedAssetPlayRate = 1.0f;
+	UPROPERTY(Transient)
+	mutable float MatchedAssetPlayRate = 1.f;
 
-	UPROPERTY()
-	TOptional<FQuat> InverseMeshToActorRotation;
+	UPROPERTY(Transient)
+	FTransform MeshToActorTransform = FTransform::Identity;
 };
 
 
