Add Pose History Pre-evaluate task so that sequencer can see the upstream pose histories

diff --git a/Engine/Plugins/Experimental/AnimNextPoseSearch/Source/AnimNextPoseSearch/Private/HistoryCollectorTrait.cpp b/Engine/Plugins/Experimental/AnimNextPoseSearch/Source/AnimNextPoseSearch/Private/HistoryCollectorTrait.cpp
--- a/Engine/Plugins/Experimental/AnimNextPoseSearch/Source/AnimNextPoseSearch/Private/HistoryCollectorTrait.cpp
+++ b/Engine/Plugins/Experimental/AnimNextPoseSearch/Source/AnimNextPoseSearch/Private/HistoryCollectorTrait.cpp
@@ -7,8 +7,8 @@
 #include "EvaluationVM/EvaluationVM.h"
 #include "GenerationTools.h"
 #include "Module/AnimNextModuleInstance.h"
-#include "PoseSearch/PoseSearchHistoryAttribute.h"
 #include "TraitCore/NodeInstance.h"
+#include "PoseHistoryEvaluation.h"
 
 namespace UE::AnimNext
 {
@@ -80,6 +80,19 @@ namespace UE::AnimNext
 	};
 }
 
+void FAnimNextHistoryCollectorPreEvaluateTask::Execute(UE::AnimNext::FEvaluationVM& VM) const
+{
+	using namespace UE::AnimNext;
+	using namespace UE::PoseSearch;
+
+	if (EnumHasAnyFlags(VM.GetFlags(), EEvaluationFlags::Bones))
+	{
+		FPoseHistoryEvaluationHelper PoseHistoryEvalHelper;
+		PoseHistoryEvalHelper.PoseHistory = PoseHistory;
+		VM.PushValue(POSEHISTORY_STACK_NAME, MakeUnique<FPoseHistoryEvaluationHelper>(PoseHistoryEvalHelper));
+	}
+}
+
 void FAnimNextHistoryCollectorTask::Execute(UE::AnimNext::FEvaluationVM& VM) const
 {
 	using namespace UE::AnimNext;
@@ -131,13 +144,6 @@ void FAnimNextHistoryCollectorTask::Execute(UE::AnimNext::FEvaluationVM& VM) con
 			PoseHistory->EvaluateComponentSpace_AnyThread(DeltaTime, ComponentSpacePoseProvider, bStoreScales,
 				SharedData->RootBoneRecoveryTime, SharedData->RootBoneTranslationRecoveryRatio, SharedData->RootBoneRotationRecoveryRatio,
 				bNeedsReset, bCacheBones, BoneIndicesWithParents);
-
-			// Push pose history onto the custom attribute for downstream tasks
-			if (FPoseHistoryAnimationAttribute* PoseHistoryAttribute = (*Keyframe)->Attributes.FindOrAdd<FPoseHistoryAnimationAttribute>(UE::PoseSearch::PoseHistoryAttributeId))
-			{
-				PoseHistoryAttribute->PoseHistory = PoseHistory;
-				PoseHistoryAttribute->ScopeObject = HostObject;
-			}
 		}
 	}
 }
@@ -205,6 +211,18 @@ namespace UE::AnimNext
 		IUpdate::PostUpdate(Context, Binding, TraitState);
 	}
 
+	void FHistoryCollectorTrait::PreEvaluate(FEvaluateTraversalContext& Context, const TTraitBinding<IEvaluate>& Binding) const
+	{
+		IEvaluate::PreEvaluate(Context, Binding);
+		FInstanceData* InstanceData = Binding.GetInstanceData<FInstanceData>();
+
+		check(InstanceData->PoseHistoryPtr.IsValid());
+		FAnimNextHistoryCollectorPreEvaluateTask Task;
+		Task.PoseHistory = InstanceData->PoseHistoryPtr.Get();
+		Task.HostObject = Context.GetHostObject();
+		Context.AppendTask(Task);
+	}
+
 	void FHistoryCollectorTrait::PostEvaluate(FEvaluateTraversalContext& Context, const TTraitBinding<IEvaluate>& Binding) const
 	{
 		IEvaluate::PostEvaluate(Context, Binding);
diff --git a/Engine/Plugins/Experimental/AnimNextPoseSearch/Source/AnimNextPoseSearch/Private/HistoryCollectorTrait.h b/Engine/Plugins/Experimental/AnimNextPoseSearch/Source/AnimNextPoseSearch/Private/HistoryCollectorTrait.h
--- a/Engine/Plugins/Experimental/AnimNextPoseSearch/Source/AnimNextPoseSearch/Private/HistoryCollectorTrait.h
+++ b/Engine/Plugins/Experimental/AnimNextPoseSearch/Source/AnimNextPoseSearch/Private/HistoryCollectorTrait.h
@@ -113,6 +113,20 @@ struct FAnimNextHistoryCollectorTraitSharedData : public FAnimNextTraitSharedDat
 	#undef TRAIT_LATENT_PROPERTIES_ENUMERATOR
 };
 
+USTRUCT()
+struct FAnimNextHistoryCollectorPreEvaluateTask : public FAnimNextEvaluationTask
+{
+	GENERATED_BODY()
+
+	DECLARE_ANIM_EVALUATION_TASK(FAnimNextHistoryCollectorPreEvaluateTask)
+
+	virtual void Execute(UE::AnimNext::FEvaluationVM& VM) const override;
+
+	UE::PoseSearch::FPoseHistory* PoseHistory = nullptr;
+	const FAnimNextHistoryCollectorTraitSharedData* SharedData = nullptr;
+	TWeakObjectPtr<const UObject> HostObject;
+};
+
 USTRUCT()
 struct FAnimNextHistoryCollectorTask : public FAnimNextEvaluationTask
 {
@@ -155,6 +169,10 @@ namespace UE::AnimNext
 		// IEvaluate impl
 		virtual void PostEvaluate(FEvaluateTraversalContext& Context, const TTraitBinding<IEvaluate>& Binding) const override;
 
+		//@TODO: this is a HACK. We're pushing a task with the pose history so that evaluation modifiers can see the pose history.
+		// "Evaluation" modifiers should have an Update callback with Scoped Interface access, just like any other trait.
+		virtual void PreEvaluate(FEvaluateTraversalContext& Context, const TTraitBinding<IEvaluate>& Binding) const override;
+
 		// IPoseHistory impl
 		virtual const UE::PoseSearch::IPoseHistory* GetPoseHistory(FExecutionContext& Context, const TTraitBinding<IPoseHistory>& Binding) const override;
 	};
diff --git a/Engine/Plugins/Experimental/AnimNextPoseSearch/Source/AnimNextPoseSearch/Private/PoseHistoryEvaluation.cpp b/Engine/Plugins/Experimental/AnimNextPoseSearch/Source/AnimNextPoseSearch/Private/PoseHistoryEvaluation.cpp
new file mode 100644
--- /dev/null
+++ b/Engine/Plugins/Experimental/AnimNextPoseSearch/Source/AnimNextPoseSearch/Private/PoseHistoryEvaluation.cpp
@@ -0,0 +1,8 @@
+// Copyright Epic Games, Inc. All Rights Reserved.
+
+#include "PoseHistoryEvaluation.h"
+
+namespace UE::PoseSearch
+{
+	const UE::AnimNext::FEvaluationVMStackName POSEHISTORY_STACK_NAME = FName(TEXT("PoseHistoryStack"));
+}
\ No newline at end of file
diff --git a/Engine/Plugins/Experimental/AnimNextPoseSearch/Source/AnimNextPoseSearch/Public/PoseHistoryEvaluation.h b/Engine/Plugins/Experimental/AnimNextPoseSearch/Source/AnimNextPoseSearch/Public/PoseHistoryEvaluation.h
new file mode 100644
--- /dev/null
+++ b/Engine/Plugins/Experimental/AnimNextPoseSearch/Source/AnimNextPoseSearch/Public/PoseHistoryEvaluation.h
@@ -0,0 +1,21 @@
+// Copyright Epic Games, Inc. All Rights Reserved.
+
+#pragma once
+
+#include "PoseSearch/PoseSearchHistory.h"
+#include "EvaluationVM/EvaluationVM.h"
+
+namespace UE::PoseSearch
+{
+	extern ANIMNEXTPOSESEARCH_API const UE::AnimNext::FEvaluationVMStackName POSEHISTORY_STACK_NAME;
+
+	struct FPoseHistoryEvaluationHelper
+	{
+		FPoseHistory* PoseHistory = nullptr;
+	};
+}
+
+namespace UE::AnimNext
+{
+	ANIM_NEXT_ENABLE_EVALUATION_STACK_USAGE(TUniquePtr<UE::PoseSearch::FPoseHistoryEvaluationHelper>);
+}
\ No newline at end of file
diff --git a/Engine/Plugins/Experimental/MovieScenePoseSearchTracks/Source/MovieScenePoseSearchTracks/Private/Sections/MovieSceneStitchAnimSection.cpp b/Engine/Plugins/Experimental/MovieScenePoseSearchTracks/Source/MovieScenePoseSearchTracks/Private/Sections/MovieSceneStitchAnimSection.cpp
--- a/Engine/Plugins/Experimental/MovieScenePoseSearchTracks/Source/MovieScenePoseSearchTracks/Private/Sections/MovieSceneStitchAnimSection.cpp
+++ b/Engine/Plugins/Experimental/MovieScenePoseSearchTracks/Source/MovieScenePoseSearchTracks/Private/Sections/MovieSceneStitchAnimSection.cpp
@@ -16,7 +16,6 @@
 #include "PoseSearch/PoseSearchDatabase.h"
 #include "PoseSearch/PoseSearchFeatureChannel_PermutationTime.h"
 #include "PoseSearch/PoseSearchHistory.h"
-#include "PoseSearch/PoseSearchHistoryAttribute.h"
 #include "PoseSearch/PoseSearchLibrary.h"
 #include "PoseSearch/PoseSearchSchema.h"
 #include "PoseSearchTracksComponentTypes.h"
@@ -24,6 +23,7 @@
 #include "Tracks/MovieSceneStitchAnimTrack.h"
 #include "Variants/MovieSceneTimeWarpVariantPayloads.h"
 #include "VisualLogger/VisualLogger.h"
+#include "PoseHistoryEvaluation.h"
 
 #include UE_INLINE_GENERATED_CPP_BY_NAME(MovieSceneStitchAnimSection)
 
@@ -62,13 +62,9 @@ void FMovieSceneStitchAnimEvaluationTask::Execute(UE::AnimNext::FEvaluationVM& V
 	// TODO: Find a way to initialize a pose history here for the first time and store it in something the task can access (stitch anim system?)
 	// if not already present.
 	const IPoseHistory* PoseHistory = nullptr;
-	if (const TUniquePtr<FKeyframeState>* Keyframe = VM.PeekValue<TUniquePtr<FKeyframeState>>(KEYFRAME_STACK_NAME, 0))
+	if (const TUniquePtr<FPoseHistoryEvaluationHelper>* PoseHistoryEvalHelper = VM.PeekValue<TUniquePtr<FPoseHistoryEvaluationHelper>>(POSEHISTORY_STACK_NAME, 0))
 	{
-		const FPoseHistoryAnimationAttribute* PoseHistoryAttribute = (*Keyframe)->Attributes.Find<FPoseHistoryAnimationAttribute>(PoseHistoryAttributeId);
-		if (PoseHistoryAttribute && PoseHistoryAttribute->ScopeObject.IsValid())
-		{
-			PoseHistory = PoseHistoryAttribute->PoseHistory;
-		}
+		PoseHistory = (*PoseHistoryEvalHelper)->PoseHistory;
 	}
 
 	// If we haven't yet run a search and found an animation, do that
@@ -85,7 +81,7 @@ void FMovieSceneStitchAnimEvaluationTask::Execute(UE::AnimNext::FEvaluationVM& V
 			FOverridePoseHistory OverridePoseHistory;
 			OverridePoseHistory.Init(PoseHistory);
 
-			// computing the overridden trajectory with all the past FTransformTrajectorySample from the PoseHistory we found via FPoseHistoryAnimationAttribute, 
+			// computing the overridden trajectory with all the past FTransformTrajectorySample from the PoseHistory we found via Pose History Stack, 
 			// and one prediction future sample computed from StitchData.TargetTransform and TimeToTarget
 			const TArray<FTransformTrajectorySample>& TrajectorySamples = PoseHistory->GetTrajectory().Samples;
 			OverridePoseHistory.Trajectory.Samples.Reserve(TrajectorySamples.Num()); // over estimating memory requirements to avoid reallocations
