MovieScenePoseSearchTracks - fix for stitch track warp time calculations

diff --git a/Engine/Plugins/Experimental/MovieScenePoseSearchTracks/Source/MovieScenePoseSearchTracks/Private/Sections/MovieSceneStitchAnimSection.cpp b/Engine/Plugins/Experimental/MovieScenePoseSearchTracks/Source/MovieScenePoseSearchTracks/Private/Sections/MovieSceneStitchAnimSection.cpp
--- a/Engine/Plugins/Experimental/MovieScenePoseSearchTracks/Source/MovieScenePoseSearchTracks/Private/Sections/MovieSceneStitchAnimSection.cpp
+++ b/Engine/Plugins/Experimental/MovieScenePoseSearchTracks/Source/MovieScenePoseSearchTracks/Private/Sections/MovieSceneStitchAnimSection.cpp
@@ -132,9 +132,9 @@ void FMovieSceneStitchAnimEvaluationTask::Execute(UE::AnimNext::FEvaluationVM& V
 					// TODO: Add mirroring support
 					//ProviderResult.bMirrored = SearchIndexAsset->IsMirrored();
 
-					// figuring out the WantedPlayRate
-					MatchedAssetPlayRate = 1.f;
-					if (StitchData.TargetPoseAsset && TimeToTarget > 0.f)
+					// figuring out the actual interval time (that will impact the animation play rate, since we'll be extracting poses using the parametric time along the duration of the stitch)
+					MatchedAssetActualIntervalTime = TimeToTarget;
+					if (StitchData.TargetPoseAsset)
 					{
 						if (const UPoseSearchFeatureChannel_PermutationTime* PermutationTimeChannel = Database->Schema->FindFirstChannelOfType<UPoseSearchFeatureChannel_PermutationTime>())
 						{
@@ -142,8 +142,7 @@ void FMovieSceneStitchAnimEvaluationTask::Execute(UE::AnimNext::FEvaluationVM& V
 							if (!SearchIndex.IsValuesEmpty())
 							{
 								TConstArrayView<float> ResultData = Database->GetSearchIndex().GetPoseValues(SearchResult.PoseIdx);
-								const float ActualIntervalTime = PermutationTimeChannel->GetPermutationTime(ResultData);
-								MatchedAssetPlayRate = ActualIntervalTime / TimeToTarget;
+								MatchedAssetActualIntervalTime = PermutationTimeChannel->GetPermutationTime(ResultData);
 							}
 						}
 					}
@@ -172,17 +171,24 @@ void FMovieSceneStitchAnimEvaluationTask::Execute(UE::AnimNext::FEvaluationVM& V
 		return;
 	}
 
-	const float AnimationTime = MatchedAssetTime + ((CurrentTime - InitialTime) * MatchedAssetPlayRate);
-	const float PreviousAnimationTime = AnimationTime - ((CurrentTime - PreviousTime) * MatchedAssetPlayRate);
-
 	// Calculate the root transform by warping linearly the animation root transforms to the target StitchData.TargetTransform (after converting it to the correct animation space in MeshComponent space)
-	const float BlendWeightDenominator = TimeToTarget + CurrentTime - InitialTime;
-	const float BlendWeight = FMath::IsNearlyZero(BlendWeightDenominator) ? 1.f : (CurrentTime - InitialTime) / BlendWeightDenominator;
+	const float TotalTime = TimeToTarget + CurrentTime - InitialTime;
+	float CurrentBlendWeight = 1.f;
+	float PreviousBlendWeight = 1.f;
+	if (TotalTime > UE_SMALL_NUMBER)
+	{
+		CurrentBlendWeight = (CurrentTime - InitialTime) / TotalTime;
+		PreviousBlendWeight = (PreviousTime - InitialTime) / TotalTime;
+	}
+
+	const float AnimationTime =  MatchedAssetTime + CurrentBlendWeight * MatchedAssetActualIntervalTime;
+	const float PreviousAnimationTime = MatchedAssetTime + PreviousBlendWeight * MatchedAssetActualIntervalTime;
+	const float AnimationTargetTime = MatchedAssetTime + MatchedAssetActualIntervalTime;
 
 	// mesh space animation transforms associate to the beginning of the stitch track AnimSpaceInitial, end AnimSpaceTarget, and current AnimSpaceCurrent
 	const FTransform AnimSpaceInitial = AnimSequence->ExtractRootTrackTransform(FAnimExtractContext((double)MatchedAssetTime), nullptr);
 	const FTransform AnimSpaceCurrent = AnimSequence->ExtractRootTrackTransform(FAnimExtractContext((double)AnimationTime), nullptr);
-	const FTransform AnimSpaceTarget = AnimSequence->ExtractRootTrackTransform(FAnimExtractContext((double)(AnimationTime + (TimeToTarget * MatchedAssetPlayRate))), nullptr);
+	const FTransform AnimSpaceTarget = AnimSequence->ExtractRootTrackTransform(FAnimExtractContext((double)AnimationTargetTime), nullptr);
 
 	// local delta transforms from the beginning of the animation to the end (InitialToTarget), and current playback time (InitialToCurrent)
 	const FTransform InitialToTarget = AnimSpaceTarget.GetRelativeTransform(AnimSpaceInitial);
@@ -201,7 +207,7 @@ void FMovieSceneStitchAnimEvaluationTask::Execute(UE::AnimNext::FEvaluationVM& V
 	const FTransform AlignmentError = FinalMeshTransform.GetRelativeTransform(FinalAnimationTransform);
 	// blending the AlignmentError by BlendWeight
 	FTransform BlendedAlignmentError;
-	BlendedAlignmentError.Blend(FTransform::Identity, AlignmentError, BlendWeight);
+	BlendedAlignmentError.Blend(FTransform::Identity, AlignmentError, CurrentBlendWeight);
 	// calculating where the mesh want to be with after applying the animation AND the blended portion from the alignement error (BlendedAlignmentError)
 	const FTransform CurrentMeshTransform = BlendedAlignmentError * CurrentAnimationTransform;
 	// converting CurrentMeshTransform to world actor space so we can set it to the FTransformAnimationAttribute
diff --git a/Engine/Plugins/Experimental/MovieScenePoseSearchTracks/Source/MovieScenePoseSearchTracks/Public/Sections/MovieSceneStitchAnimSection.h b/Engine/Plugins/Experimental/MovieScenePoseSearchTracks/Source/MovieScenePoseSearchTracks/Public/Sections/MovieSceneStitchAnimSection.h
--- a/Engine/Plugins/Experimental/MovieScenePoseSearchTracks/Source/MovieScenePoseSearchTracks/Public/Sections/MovieSceneStitchAnimSection.h
+++ b/Engine/Plugins/Experimental/MovieScenePoseSearchTracks/Source/MovieScenePoseSearchTracks/Public/Sections/MovieSceneStitchAnimSection.h
@@ -90,10 +90,10 @@ struct FMovieSceneStitchAnimEvaluationTask : public FAnimNextEvaluationTask
 	mutable float MatchedAssetTime = 0.f;
 
 	UPROPERTY(Transient)
-	mutable float MatchedAssetPlayRate = 1.f;
+	FTransform MeshToActorTransform = FTransform::Identity;
 
 	UPROPERTY(Transient)
-	FTransform MeshToActorTransform = FTransform::Identity;
+	mutable float MatchedAssetActualIntervalTime = 0.f;
 };
 
 
