Added a bool to toggle authoritative root motion to stitch tracks

diff --git a/Engine/Plugins/Experimental/MovieScenePoseSearchTracks/Source/MovieScenePoseSearchTracks/Private/Sections/MovieSceneStitchAnimSection.cpp b/Engine/Plugins/Experimental/MovieScenePoseSearchTracks/Source/MovieScenePoseSearchTracks/Private/Sections/MovieSceneStitchAnimSection.cpp
--- a/Engine/Plugins/Experimental/MovieScenePoseSearchTracks/Source/MovieScenePoseSearchTracks/Private/Sections/MovieSceneStitchAnimSection.cpp
+++ b/Engine/Plugins/Experimental/MovieScenePoseSearchTracks/Source/MovieScenePoseSearchTracks/Private/Sections/MovieSceneStitchAnimSection.cpp
@@ -244,8 +244,11 @@ void FMovieSceneStitchAnimEvaluationTask::Execute(UE::AnimNext::FEvaluationVM& V
 		FFloatAnimationAttribute* RootMotionWeightAttribute = Keyframe.Attributes.FindOrAdd<FFloatAnimationAttribute>(AnimMixerComponents->RootTransformWeightAttributeId);
 		RootMotionWeightAttribute->Value = 1.f;
 
-		FIntegerAnimationAttribute* RootMotionIsAuthoritativeAttribute = Keyframe.Attributes.FindOrAdd<FIntegerAnimationAttribute>(AnimMixerComponents->RootTransformIsAuthoritativeAttributeId);
-		RootMotionIsAuthoritativeAttribute->Value = 1;
+		if (StitchData.bAuthoritativeRootMotion)
+		{
+			FIntegerAnimationAttribute* RootMotionIsAuthoritativeAttribute = Keyframe.Attributes.FindOrAdd<FIntegerAnimationAttribute>(AnimMixerComponents->RootTransformIsAuthoritativeAttributeId);
+			RootMotionIsAuthoritativeAttribute->Value = 1;
+		}
 	}
 
 	VM.PushValue(KEYFRAME_STACK_NAME, MakeUnique<FKeyframeState>(MoveTemp(Keyframe)));
@@ -339,6 +342,7 @@ void UMovieSceneStitchAnimSection::ImportEntityImpl(UMovieSceneEntitySystemLinke
 	StitchData.StartFrame = GetInclusiveStartFrame();
 	StitchData.EndFrame = GetExclusiveEndFrame();
 	StitchData.TargetTransformSpace = TargetTransformSpace;
+	StitchData.bAuthoritativeRootMotion = bAuthoritativeRootMotion;
 	
 	// Make a default eval task. This will be filled out more by the stitch system
 	TSharedPtr<FMovieSceneStitchAnimEvaluationTask> EvalTask = MakeShared<FMovieSceneStitchAnimEvaluationTask>();
diff --git a/Engine/Plugins/Experimental/MovieScenePoseSearchTracks/Source/MovieScenePoseSearchTracks/Public/Sections/MovieSceneStitchAnimSection.h b/Engine/Plugins/Experimental/MovieScenePoseSearchTracks/Source/MovieScenePoseSearchTracks/Public/Sections/MovieSceneStitchAnimSection.h
--- a/Engine/Plugins/Experimental/MovieScenePoseSearchTracks/Source/MovieScenePoseSearchTracks/Public/Sections/MovieSceneStitchAnimSection.h
+++ b/Engine/Plugins/Experimental/MovieScenePoseSearchTracks/Source/MovieScenePoseSearchTracks/Public/Sections/MovieSceneStitchAnimSection.h
@@ -41,6 +41,9 @@ struct FMovieSceneStitchAnimComponentData
 	UPROPERTY()
 	EMovieSceneRootMotionSpace TargetTransformSpace = EMovieSceneRootMotionSpace::AnimationSpace;
 
+	UPROPERTY()
+	bool bAuthoritativeRootMotion = true;
+
 	double MapTimeToSectionSeconds(FFrameTime InPosition, FFrameRate InFrameRate) const;
 };
 
@@ -138,6 +141,9 @@ public:
 	UPROPERTY(EditAnywhere, Category = "Stitch")
 	EMovieSceneRootMotionSpace TargetTransformSpace = EMovieSceneRootMotionSpace::AnimationSpace;
 
+	UPROPERTY(EditAnywhere, Category = "Stitch")
+	bool bAuthoritativeRootMotion = true;
+
 	//~ UMovieSceneSection interface
 	virtual EMovieSceneChannelProxyType CacheChannelProxy() override;
 	virtual UObject* GetSourceObject() const override;
