Sequencer: Minor Stitch Track UX fixes.
* Animation producing sections that go on the Anim Mixer can now choose a custom tint for when they're on the mixer
* The stitch section now shows text describing what stitch database is being used
* The stitch section has a default length now.

diff --git a/Engine/Plugins/Experimental/MovieSceneAnimMixer/Source/MovieSceneAnimMixer/Internal/MovieSceneAnimationMixerTrack.h b/Engine/Plugins/Experimental/MovieSceneAnimMixer/Source/MovieSceneAnimMixer/Internal/MovieSceneAnimationMixerTrack.h
--- a/Engine/Plugins/Experimental/MovieSceneAnimMixer/Source/MovieSceneAnimMixer/Internal/MovieSceneAnimationMixerTrack.h
+++ b/Engine/Plugins/Experimental/MovieSceneAnimMixer/Source/MovieSceneAnimMixer/Internal/MovieSceneAnimationMixerTrack.h
@@ -14,7 +14,6 @@
 #include "Tracks/MovieSceneCommonAnimationTrack.h"
 #include "MovieSceneAnimationMixerTrack.generated.h"
 
-
 // Enum describing what space that root motion should be applied in.
 UENUM()
 enum class EMovieSceneRootMotionSpace : uint8
@@ -74,6 +73,9 @@ public:
 	GENERATED_BODY()
 
 	virtual int32 GetRowSortOrder() const = 0;
+
+	// Allows different section types to have different colors when used in the mixer, without needing to override their colors everywhere.
+	virtual FColor GetMixerSectionTint() const = 0;
 };
 
 UCLASS(MinimalAPI)
@@ -92,8 +94,16 @@ class UMovieSceneAnimationSectionDecoration
 		return RowSortOrder;
 	}
 
+	virtual FColor GetMixerSectionTint() const override
+	{
+		return MixerTintOverride;
+	}
+
 	UPROPERTY()
 	int32 RowSortOrder = 0;
+
+private:
+	FColor MixerTintOverride = FColor();
 };
 
 UCLASS(MinimalAPI)
diff --git a/Engine/Plugins/Experimental/MovieSceneAnimMixer/Source/MovieSceneAnimMixer/Internal/MovieSceneRootMotionSection.h b/Engine/Plugins/Experimental/MovieSceneAnimMixer/Source/MovieSceneAnimMixer/Internal/MovieSceneRootMotionSection.h
--- a/Engine/Plugins/Experimental/MovieSceneAnimMixer/Source/MovieSceneAnimMixer/Internal/MovieSceneRootMotionSection.h
+++ b/Engine/Plugins/Experimental/MovieSceneAnimMixer/Source/MovieSceneAnimMixer/Internal/MovieSceneRootMotionSection.h
@@ -41,6 +41,11 @@ public:
 
 	virtual int32 GetRowSortOrder() const;
 
+	virtual FColor GetMixerSectionTint() const override
+	{
+		return MixerTintOverride;
+	}
+
 	virtual EMovieSceneChannelProxyType CacheChannelProxy() override;
 
 	virtual void ImportEntityImpl(UMovieSceneEntitySystemLinker* EntityLinker, const FEntityImportParams& Params, FImportedEntity* OutImportedEntity) override;
@@ -50,4 +55,6 @@ private:
 
 	UPROPERTY()
 	FMovieSceneByteChannel RootDestinationChannel;
+
+	FColor MixerTintOverride = FColor(20, 70, 70, 200);
 };
diff --git a/Engine/Plugins/Experimental/MovieSceneAnimMixer/Source/MovieSceneAnimMixer/MovieSceneAnimMixer.Build.cs b/Engine/Plugins/Experimental/MovieSceneAnimMixer/Source/MovieSceneAnimMixer/MovieSceneAnimMixer.Build.cs
--- a/Engine/Plugins/Experimental/MovieSceneAnimMixer/Source/MovieSceneAnimMixer/MovieSceneAnimMixer.Build.cs
+++ b/Engine/Plugins/Experimental/MovieSceneAnimMixer/Source/MovieSceneAnimMixer/MovieSceneAnimMixer.Build.cs
@@ -34,7 +34,7 @@ namespace UnrealBuildTool.Rules
 				PrivateDependencyModuleNames.AddRange(new string[]
 					{
 						"EditorFramework",
-						"UnrealEd"
+						"UnrealEd",
 					});
 			}
 		}
diff --git a/Engine/Plugins/Experimental/MovieSceneAnimMixer/Source/MovieSceneAnimMixer/Private/MovieSceneAnimationMixerTrack.cpp b/Engine/Plugins/Experimental/MovieSceneAnimMixer/Source/MovieSceneAnimMixer/Private/MovieSceneAnimationMixerTrack.cpp
--- a/Engine/Plugins/Experimental/MovieSceneAnimMixer/Source/MovieSceneAnimMixer/Private/MovieSceneAnimationMixerTrack.cpp
+++ b/Engine/Plugins/Experimental/MovieSceneAnimMixer/Source/MovieSceneAnimMixer/Private/MovieSceneAnimationMixerTrack.cpp
@@ -269,6 +269,11 @@ void UMovieSceneAnimationMixerTrack::OnSectionAddedImpl(UMovieSceneSection* Sect
 		UMovieSceneAnimationBaseTransformDecoration* BaseTransform = NewObject<UMovieSceneAnimationBaseTransformDecoration>(AnimSection, UMovieSceneAnimationBaseTransformDecoration::StaticClass(), NAME_None, RF_Transactional);
 		AnimSection->AddDecoration(BaseTransform);
 	}
+
+	if (IMovieSceneAnimationSectionInterface* AnimationSectionInterface = Cast<IMovieSceneAnimationSectionInterface>(Section))
+	{
+		Section->SetColorTint(AnimationSectionInterface->GetMixerSectionTint());
+	}
 }
 
 
diff --git a/Engine/Plugins/Experimental/MovieSceneAnimMixer/Source/MovieSceneAnimMixerEditor/Private/MovieSceneAnimationMixerTrackEditor.cpp b/Engine/Plugins/Experimental/MovieSceneAnimMixer/Source/MovieSceneAnimMixerEditor/Private/MovieSceneAnimationMixerTrackEditor.cpp
--- a/Engine/Plugins/Experimental/MovieSceneAnimMixer/Source/MovieSceneAnimMixerEditor/Private/MovieSceneAnimationMixerTrackEditor.cpp
+++ b/Engine/Plugins/Experimental/MovieSceneAnimMixer/Source/MovieSceneAnimMixerEditor/Private/MovieSceneAnimationMixerTrackEditor.cpp
@@ -210,6 +210,8 @@ public:
 	}
 };
 
+TMap<const UClass*, FOnMakeSectionInterfaceDelegate> FAnimationMixerTrackEditor::MakeSectionInterfaceCallbacks;
+
 TSharedRef<ISequencerTrackEditor> FAnimationMixerTrackEditor::CreateTrackEditor(TSharedRef<ISequencer> InSequencer)
 {
 	return MakeShareable(new FAnimationMixerTrackEditor(InSequencer));
@@ -296,6 +298,14 @@ TSharedRef<ISequencerSection> FAnimationMixerTrackEditor::MakeSectionInterface(U
 		return MakeShared<FAnimMixerAnimationSection>(SectionObject, GetSequencer());
 	}
 
+	if (FOnMakeSectionInterfaceDelegate* MakeSectionInterfaceDelegate = MakeSectionInterfaceCallbacks.Find(SectionObject.GetClass()))
+	{
+		if ((*MakeSectionInterfaceDelegate).IsBound())
+		{
+			return (*MakeSectionInterfaceDelegate).Execute(SectionObject, Track, ObjectBinding);
+		}
+	}
+
 	return FMovieSceneTrackEditor::MakeSectionInterface(SectionObject, Track, ObjectBinding);
 }
 
diff --git a/Engine/Plugins/Experimental/MovieSceneAnimMixer/Source/MovieSceneAnimMixerEditor/Private/MovieSceneAnimationMixerTrackEditor.h b/Engine/Plugins/Experimental/MovieSceneAnimMixer/Source/MovieSceneAnimMixerEditor/Private/MovieSceneAnimationMixerTrackEditor.h
--- a/Engine/Plugins/Experimental/MovieSceneAnimMixer/Source/MovieSceneAnimMixerEditor/Private/MovieSceneAnimationMixerTrackEditor.h
+++ b/Engine/Plugins/Experimental/MovieSceneAnimMixer/Source/MovieSceneAnimMixerEditor/Private/MovieSceneAnimationMixerTrackEditor.h
@@ -7,10 +7,11 @@
 #include "MovieSceneTrackEditor.h"
 #include "TrackEditors/CommonAnimationTrackEditor.h"
 
+DECLARE_DELEGATE_RetVal_ThreeParams(TSharedRef<ISequencerSection>, FOnMakeSectionInterfaceDelegate, UMovieSceneSection&, UMovieSceneTrack&, FGuid);
+
 namespace UE::Sequencer
 {
 
-
 class FAnimationMixerTrackEditor
 	: public FCommonAnimationTrackEditor
 {
@@ -18,6 +19,17 @@ public:
 
 	static TSharedRef<ISequencerTrackEditor> CreateTrackEditor(TSharedRef<ISequencer> InSequencer);
 
+
+	static void RegisterCustomMixerAnimSection(const UClass* InClass, FOnMakeSectionInterfaceDelegate InMakeSectionInterfaceDelegate)
+	{
+		MakeSectionInterfaceCallbacks.Add(TTuple<const UClass*, FOnMakeSectionInterfaceDelegate>{InClass, InMakeSectionInterfaceDelegate});
+	}
+
+	static void UnregisterCustomMixerAnimSection(const UClass* InClass)
+	{
+		MakeSectionInterfaceCallbacks.Remove(InClass);
+	}
+
 public:
 
 	FAnimationMixerTrackEditor(TSharedRef<ISequencer> InSequencer);
@@ -43,6 +55,8 @@ private:
 	bool FilterAnimSequences(const FAssetData& AssetData, USkeleton* Skeleton);
 	void OnAnimationAssetSelected(const FAssetData& AssetData, TArray<FGuid> ObjectBindings, UMovieSceneTrack* Track);
 	void OnAnimationAssetEnterPressed(const TArray<FAssetData>& AssetData, TArray<FGuid> ObjectBindings, UMovieSceneTrack* Track);
+
+	MOVIESCENEANIMMIXEREDITOR_API static TMap<const UClass*, FOnMakeSectionInterfaceDelegate> MakeSectionInterfaceCallbacks;
 };
 
 
diff --git a/Engine/Plugins/Experimental/MovieScenePoseSearchTracks/Source/MovieScenePoseSearchTracks/Private/Sections/MovieSceneStitchAnimSection.cpp b/Engine/Plugins/Experimental/MovieScenePoseSearchTracks/Source/MovieScenePoseSearchTracks/Private/Sections/MovieSceneStitchAnimSection.cpp
--- a/Engine/Plugins/Experimental/MovieScenePoseSearchTracks/Source/MovieScenePoseSearchTracks/Private/Sections/MovieSceneStitchAnimSection.cpp
+++ b/Engine/Plugins/Experimental/MovieScenePoseSearchTracks/Source/MovieScenePoseSearchTracks/Private/Sections/MovieSceneStitchAnimSection.cpp
@@ -8,6 +8,8 @@
 #include "AnimMixerComponentTypes.h"
 #include "Channels/MovieSceneChannelProxy.h"
 #include "DecompressionTools.h"
+#include "MovieScene.h"
+#include "MovieSceneSequence.h"
 #include "EntitySystem/BuiltInComponentTypes.h"
 #include "EvaluationVM/EvaluationVM.h"
 #include "EvaluationVM/KeyframeState.h"
@@ -285,6 +287,19 @@ UMovieSceneStitchAnimSection::UMovieSceneStitchAnimSection(const FObjectInitiali
 	: Super(ObjectInitializer)
 {
 	BlendType = EMovieSceneBlendType::Absolute;
+
+	UMovieSceneSequence* OuterSequence = GetTypedOuter<UMovieSceneSequence>();
+
+	if (OuterSequence)
+	{
+		FFrameRate TickResolution = GetTypedOuter<UMovieSceneSequence>()->GetMovieScene()->GetTickResolution();
+		FFrameRate DisplayRate = GetTypedOuter<UMovieSceneSequence>()->GetMovieScene()->GetDisplayRate();
+
+		FFrameTime EndTime = ConvertFrameTime(FFrameTime(5), DisplayRate, TickResolution);
+
+		SectionRange.Value = TRange<FFrameNumber>(FFrameNumber(0), EndTime.FrameNumber);
+	}
+	
 }
 
 EMovieSceneChannelProxyType UMovieSceneStitchAnimSection::CacheChannelProxy()
diff --git a/Engine/Plugins/Experimental/MovieScenePoseSearchTracks/Source/MovieScenePoseSearchTracks/Public/Sections/MovieSceneStitchAnimSection.h b/Engine/Plugins/Experimental/MovieScenePoseSearchTracks/Source/MovieScenePoseSearchTracks/Public/Sections/MovieSceneStitchAnimSection.h
--- a/Engine/Plugins/Experimental/MovieScenePoseSearchTracks/Source/MovieScenePoseSearchTracks/Public/Sections/MovieSceneStitchAnimSection.h
+++ b/Engine/Plugins/Experimental/MovieScenePoseSearchTracks/Source/MovieScenePoseSearchTracks/Public/Sections/MovieSceneStitchAnimSection.h
@@ -153,6 +153,11 @@ public:
 		return 0;
 	}
 
+	virtual FColor GetMixerSectionTint() const override
+	{
+		return MixerTintOverride;
+	}
+
 protected:
 
 	virtual float GetTotalWeightValue(FFrameTime InTime) const override;
@@ -160,4 +165,6 @@ protected:
 	/** ~IMovieSceneEntityProvider interface */
 	virtual void ImportEntityImpl(UMovieSceneEntitySystemLinker* EntityLinker, const FEntityImportParams& InParams, FImportedEntity* OutImportedEntity) override;
 
+	FColor MixerTintOverride = FColor(20, 70, 120, 200);
+
 };
diff --git a/Engine/Plugins/Experimental/MovieScenePoseSearchTracks/Source/MovieScenePoseSearchTracksEditor/MovieScenePoseSearchTracksEditor.Build.cs b/Engine/Plugins/Experimental/MovieScenePoseSearchTracks/Source/MovieScenePoseSearchTracksEditor/MovieScenePoseSearchTracksEditor.Build.cs
--- a/Engine/Plugins/Experimental/MovieScenePoseSearchTracks/Source/MovieScenePoseSearchTracksEditor/MovieScenePoseSearchTracksEditor.Build.cs
+++ b/Engine/Plugins/Experimental/MovieScenePoseSearchTracks/Source/MovieScenePoseSearchTracksEditor/MovieScenePoseSearchTracksEditor.Build.cs
@@ -16,6 +16,7 @@ namespace UnrealBuildTool.Rules
 					"MovieSceneTracks",
 					"MovieSceneTools",
 					"MovieScenePoseSearchTracks",
+					"MovieSceneAnimMixerEditor",
 					"AnimGraphRuntime",
 					"UAF",
 					"UAFAnimGraph",
diff --git a/Engine/Plugins/Experimental/MovieScenePoseSearchTracks/Source/MovieScenePoseSearchTracksEditor/Private/MovieScenePoseSearchTracksEditorModule.cpp b/Engine/Plugins/Experimental/MovieScenePoseSearchTracks/Source/MovieScenePoseSearchTracksEditor/Private/MovieScenePoseSearchTracksEditorModule.cpp
--- a/Engine/Plugins/Experimental/MovieScenePoseSearchTracks/Source/MovieScenePoseSearchTracksEditor/Private/MovieScenePoseSearchTracksEditorModule.cpp
+++ b/Engine/Plugins/Experimental/MovieScenePoseSearchTracks/Source/MovieScenePoseSearchTracksEditor/Private/MovieScenePoseSearchTracksEditorModule.cpp
@@ -4,7 +4,9 @@
 #include "Modules/ModuleInterface.h"
 #include "Modules/ModuleManager.h"
 #include "ISequencerModule.h"
+#include "MovieSceneAnimMixerEditor/Private/MovieSceneAnimationMixerTrackEditor.h"
 #include "TrackEditors/StitchAnimTrackEditor.h"
+#include "Sections/MovieSceneStitchAnimSection.h"
 
 namespace UE::MovieScene
 {
@@ -13,17 +15,28 @@ void FMovieScenePoseSearchTracksEditorModule::StartupModule()
 {
 	using namespace UE::Sequencer;
 
-	if (GIsEditor)
-	{
-		ISequencerModule& SequencerModule = FModuleManager::Get().LoadModuleChecked<ISequencerModule>("Sequencer");
-
-		// register specialty track editors
-		StitchAnimationTrackCreateEditorHandle = SequencerModule.RegisterTrackEditor(FOnCreateTrackEditor::CreateStatic(&FStitchAnimTrackEditor::CreateTrackEditor));
-	}
+	ISequencerModule& SequencerModule = FModuleManager::Get().LoadModuleChecked<ISequencerModule>("Sequencer");
+
+	// register specialty track editors
+	StitchAnimationTrackCreateEditorHandle = SequencerModule.RegisterTrackEditor(FOnCreateTrackEditor::CreateStatic(&FStitchAnimTrackEditor::CreateTrackEditor));
+	
+	FAnimationMixerTrackEditor::RegisterCustomMixerAnimSection(UMovieSceneStitchAnimSection::StaticClass(), FOnMakeSectionInterfaceDelegate::CreateLambda([](UMovieSceneSection& Section, UMovieSceneTrack& Track, FGuid Guid)
+		{
+			if (UMovieSceneStitchAnimSection* StitchSection = Cast<UMovieSceneStitchAnimSection>(&Section))
+			{
+				TSharedPtr<ISequencerSection> SectionPtr = MakeShareable(new FStitchAnimSection(*StitchSection));
+				return SectionPtr.ToSharedRef();
+			}
+			TSharedPtr<ISequencerSection> SectionPtr = MakeShareable(new FSequencerSection(Section));
+			return SectionPtr.ToSharedRef();
+		}));
+	
 }
 
 void FMovieScenePoseSearchTracksEditorModule::ShutdownModule()
 {
+	using namespace UE::Sequencer;
+
 	if (!FModuleManager::Get().IsModuleLoaded("Sequencer"))
 	{
 		return;
@@ -31,6 +44,8 @@ void FMovieScenePoseSearchTracksEditorModule::ShutdownModule()
 
 	ISequencerModule& SequencerModule = FModuleManager::Get().GetModuleChecked<ISequencerModule>("Sequencer");
 
+	FAnimationMixerTrackEditor::UnregisterCustomMixerAnimSection(UMovieSceneStitchAnimSection::StaticClass());
+
 	// unregister specialty track editors
 	SequencerModule.UnRegisterTrackEditor(StitchAnimationTrackCreateEditorHandle);
 }
diff --git a/Engine/Plugins/Experimental/MovieScenePoseSearchTracks/Source/MovieScenePoseSearchTracksEditor/Private/TrackEditors/StitchAnimTrackEditor.cpp b/Engine/Plugins/Experimental/MovieScenePoseSearchTracks/Source/MovieScenePoseSearchTracksEditor/Private/TrackEditors/StitchAnimTrackEditor.cpp
--- a/Engine/Plugins/Experimental/MovieScenePoseSearchTracks/Source/MovieScenePoseSearchTracksEditor/Private/TrackEditors/StitchAnimTrackEditor.cpp
+++ b/Engine/Plugins/Experimental/MovieScenePoseSearchTracks/Source/MovieScenePoseSearchTracksEditor/Private/TrackEditors/StitchAnimTrackEditor.cpp
@@ -170,9 +170,8 @@ USkeleton* AcquireSkeletonFromObjectGuid(const FGuid& Guid, TSharedPtr<ISequence
 	return nullptr;
 }
 
-FStitchAnimSection::FStitchAnimSection(UMovieSceneSection& InSection, TWeakPtr<ISequencer> InSequencer)
+FStitchAnimSection::FStitchAnimSection(UMovieSceneSection& InSection)
 	: Section(*CastChecked<UMovieSceneStitchAnimSection>(&InSection))
-	, Sequencer(InSequencer)
 {
 }
 
@@ -245,7 +244,7 @@ TSharedRef<ISequencerSection> FStitchAnimTrackEditor::MakeSectionInterface(UMovi
 {
 	check(SupportsType(SectionObject.GetOuter()->GetClass()));
 
-	return MakeShareable(new FStitchAnimSection(SectionObject, GetSequencer()));
+	return MakeShareable(new FStitchAnimSection(SectionObject));
 }
 
 bool FStitchAnimTrackEditor::HandleAssetAdded(UObject* Asset, const FGuid& TargetObjectGuid)
diff --git a/Engine/Plugins/Experimental/MovieScenePoseSearchTracks/Source/MovieScenePoseSearchTracksEditor/Private/TrackEditors/StitchAnimTrackEditor.h b/Engine/Plugins/Experimental/MovieScenePoseSearchTracks/Source/MovieScenePoseSearchTracksEditor/Private/TrackEditors/StitchAnimTrackEditor.h
--- a/Engine/Plugins/Experimental/MovieScenePoseSearchTracks/Source/MovieScenePoseSearchTracksEditor/Private/TrackEditors/StitchAnimTrackEditor.h
+++ b/Engine/Plugins/Experimental/MovieScenePoseSearchTracks/Source/MovieScenePoseSearchTracksEditor/Private/TrackEditors/StitchAnimTrackEditor.h
@@ -89,7 +89,7 @@ class FStitchAnimSection
 public:
 
 	/** Constructor. */
-	FStitchAnimSection(UMovieSceneSection& InSection, TWeakPtr<ISequencer> InSequencer);
+	FStitchAnimSection(UMovieSceneSection& InSection);
 
 	/** Virtual destructor. */
 	virtual ~FStitchAnimSection();
@@ -110,8 +110,5 @@ private:
 	/** The section we are visualizing */
 	UMovieSceneStitchAnimSection& Section;
 
-	/** Used to draw animation frame, need selection state and local time*/
-	TWeakPtr<ISequencer> Sequencer;
-
 	TUniquePtr<FMovieSceneSequenceTransform> InitialDragTransform;
 };
