Sequencer: Add a way for stitch track to work with Animation Blueprints. This is a bit of a workaround for the issue of pose history not being able to be both read from and written to easily via an attribute, which is necessary for the stitch track. Now, the user can add an 'OverridePoseHistoryName' which, in animation blueprint mode, will look for a pose history node tagged with that name and pull the pose history from it for use in the stitch track. This allows the pose history node to be _after_ the sequencer target node, which is necessary for blending out to work correctly. In addition, we set up an example properly in the AnimSandbox to show this working.

#virtualized

diff --git a/Engine/Plugins/Experimental/MovieScenePoseSearchTracks/Source/MovieScenePoseSearchTracks/Private/Sections/MovieSceneStitchAnimSection.cpp b/Engine/Plugins/Experimental/MovieScenePoseSearchTracks/Source/MovieScenePoseSearchTracks/Private/Sections/MovieSceneStitchAnimSection.cpp
--- a/Engine/Plugins/Experimental/MovieScenePoseSearchTracks/Source/MovieScenePoseSearchTracks/Private/Sections/MovieSceneStitchAnimSection.cpp
+++ b/Engine/Plugins/Experimental/MovieScenePoseSearchTracks/Source/MovieScenePoseSearchTracks/Private/Sections/MovieSceneStitchAnimSection.cpp
@@ -26,6 +26,7 @@
 #include "Variants/MovieSceneTimeWarpVariantPayloads.h"
 #include "VisualLogger/VisualLogger.h"
 #include "PoseHistoryEvaluation.h"
+#include "PoseSearch/AnimNode_PoseSearchHistoryCollector.h"
 
 #include UE_INLINE_GENERATED_CPP_BY_NAME(MovieSceneStitchAnimSection)
 
@@ -69,6 +70,16 @@ void FMovieSceneStitchAnimEvaluationTask::Execute(UE::UAF::FEvaluationVM& VM) co
 		PoseHistory = (*PoseHistoryEvalHelper)->PoseHistoryPtr.Get();
 	}
 
+	const UObject* AnimContext = ContextObject.Get();
+	const UAnimInstance* AnimInstance = Cast<UAnimInstance>(AnimContext);
+	if (PoseHistory == nullptr && AnimInstance != nullptr && !StitchData.OverridePoseHistoryTag.IsNone())
+	{
+		if (const FAnimNode_PoseSearchHistoryCollector_Base* PoseHistoryNode = UPoseSearchLibrary::FindPoseHistoryNode(StitchData.OverridePoseHistoryTag, AnimInstance))
+		{
+			PoseHistory = &PoseHistoryNode->GetPoseHistory();
+		}
+	}
+
 	// If we haven't yet run a search and found an animation, do that
 	if (!MatchedAsset)
 	{
@@ -119,7 +130,6 @@ void FMovieSceneStitchAnimEvaluationTask::Execute(UE::UAF::FEvaluationVM& VM) co
 			ContinuingProperties.PlayingAssetAccumulatedTime = PlayingAsset.AccumulatedTime;*/
 
 			const UObject* StitchDatabaseObj = StitchData.StitchDatabase.Get();
-			const UObject* AnimContext = ContextObject.Get();
 			FSearchResults_Single SearchResults;
 			UPoseSearchLibrary::MotionMatch(MakeArrayView(&AnimContext, 1), MakeArrayView(&DefaultRole, 1), MakeArrayView(&PoseHistory, 1),
 				MakeArrayView(&StitchDatabaseObj, 1), ContinuingProperties, TimeToTarget, FPoseSearchEvent(), SearchResults);
@@ -360,6 +370,7 @@ void UMovieSceneStitchAnimSection::ImportEntityImpl(UMovieSceneEntitySystemLinke
 	StitchData.EndFrame = GetExclusiveEndFrame();
 	StitchData.TargetTransformSpace = TargetTransformSpace;
 	StitchData.bAuthoritativeRootMotion = bAuthoritativeRootMotion;
+	StitchData.OverridePoseHistoryTag = OverridePoseHistoryTag;
 	
 	// Make a default eval task. This will be filled out more by the stitch system
 	TSharedPtr<FMovieSceneStitchAnimEvaluationTask> EvalTask = MakeShared<FMovieSceneStitchAnimEvaluationTask>();
diff --git a/Engine/Plugins/Experimental/MovieScenePoseSearchTracks/Source/MovieScenePoseSearchTracks/Private/Systems/MovieSceneStitchAnimSystem.cpp b/Engine/Plugins/Experimental/MovieScenePoseSearchTracks/Source/MovieScenePoseSearchTracks/Private/Systems/MovieSceneStitchAnimSystem.cpp
--- a/Engine/Plugins/Experimental/MovieScenePoseSearchTracks/Source/MovieScenePoseSearchTracks/Private/Systems/MovieSceneStitchAnimSystem.cpp
+++ b/Engine/Plugins/Experimental/MovieScenePoseSearchTracks/Source/MovieScenePoseSearchTracks/Private/Systems/MovieSceneStitchAnimSystem.cpp
@@ -90,7 +90,10 @@ namespace UE::MovieScene
 			// Fall back to using the skeletal mesh component
 			if (!ContextObject)
 			{
-				ContextObject = TryFindComponent.template operator()<USkeletalMeshComponent>(BoundObject);
+				if (const USkeletalMeshComponent* MeshComponent = TryFindComponent.template operator()<USkeletalMeshComponent>(BoundObject))
+				{
+					ContextObject = MeshComponent->GetAnimInstance();
+				}
 			}
 
 			return ContextObject;
diff --git a/Engine/Plugins/Experimental/MovieScenePoseSearchTracks/Source/MovieScenePoseSearchTracks/Public/Sections/MovieSceneStitchAnimSection.h b/Engine/Plugins/Experimental/MovieScenePoseSearchTracks/Source/MovieScenePoseSearchTracks/Public/Sections/MovieSceneStitchAnimSection.h
--- a/Engine/Plugins/Experimental/MovieScenePoseSearchTracks/Source/MovieScenePoseSearchTracks/Public/Sections/MovieSceneStitchAnimSection.h
+++ b/Engine/Plugins/Experimental/MovieScenePoseSearchTracks/Source/MovieScenePoseSearchTracks/Public/Sections/MovieSceneStitchAnimSection.h
@@ -43,6 +43,9 @@ struct FMovieSceneStitchAnimComponentData
 
 	UPROPERTY()
 	bool bAuthoritativeRootMotion = true;
+	
+	UPROPERTY()
+	FName OverridePoseHistoryTag;
 
 	double MapTimeToSectionSeconds(FFrameTime InPosition, FFrameRate InFrameRate) const;
 };
@@ -144,6 +147,12 @@ public:
 	UPROPERTY(EditAnywhere, Category = "Stitch")
 	bool bAuthoritativeRootMotion = true;
 
+	// When using animation blueprints, can be used to search for a specific pose history node by tag.
+	// This allow you to have a pose history node pre-sequencer for blending into sequences,
+	// but still capture a separate pose history afterwards for blending out.
+	UPROPERTY(EditAnywhere, category = "Stitch")
+	FName OverridePoseHistoryTag;
+
 	//~ UMovieSceneSection interface
 	virtual EMovieSceneChannelProxyType CacheChannelProxy() override;
 	virtual UObject* GetSourceObject() const override;
