Sequencer: Fix default AddSection path for Animation Mixer to properly assign sections to layers

Section types without a custom UX asset picker were going through a default AddSection code path that didn't fire lifecycle events (OnSectionAddedImpl). The mixer track uses these events to assign sections to a UMovieSceneAnimationMixerLayer, so sections added via the default path were never assigned to a layer and didn't appear in the UI. Removed the workaround that the generated anims were using to avoid this bug. Also fixed an issue with the stitch track version as noticed in testing.

#Synced-CL 52466902

diff --git a/Engine/Plugins/Experimental/MovieSceneAnimMixer/Source/MovieSceneAnimMixer/Private/MovieSceneAnimationMixerTrack.cpp b/Engine/Plugins/Experimental/MovieSceneAnimMixer/Source/MovieSceneAnimMixer/Private/MovieSceneAnimationMixerTrack.cpp
--- a/Engine/Plugins/Experimental/MovieSceneAnimMixer/Source/MovieSceneAnimMixer/Private/MovieSceneAnimationMixerTrack.cpp
+++ b/Engine/Plugins/Experimental/MovieSceneAnimMixer/Source/MovieSceneAnimMixer/Private/MovieSceneAnimationMixerTrack.cpp
@@ -1305,8 +1305,9 @@ void UMovieSceneAnimationMixerTrack::AddSectionWithAutoRow(UMovieSceneSection* S
 	// Assign next available row
 	Section->SetRowIndex(NextRowIndex);
 
-	// Add the section
-	AddSection(*Section);
+	// Add the section via FSectionParameter to ensure OnSectionAddedImpl is called,
+	// which assigns the section to a mixer layer so it appears in the UI.
+	UMovieSceneTrack::AddSection(FSectionParameter{*Section});
 }
 
 int32 UMovieSceneAnimationMixerTrack::ComputeNextAvailableRowIndex() const
diff --git a/Engine/Plugins/Experimental/MovieSceneAnimMixer/Source/MovieSceneAnimMixerEditor/Private/MovieSceneAnimationMixerTrackEditor.cpp b/Engine/Plugins/Experimental/MovieSceneAnimMixer/Source/MovieSceneAnimMixerEditor/Private/MovieSceneAnimationMixerTrackEditor.cpp
--- a/Engine/Plugins/Experimental/MovieSceneAnimMixer/Source/MovieSceneAnimMixerEditor/Private/MovieSceneAnimationMixerTrackEditor.cpp
+++ b/Engine/Plugins/Experimental/MovieSceneAnimMixer/Source/MovieSceneAnimMixerEditor/Private/MovieSceneAnimationMixerTrackEditor.cpp
@@ -811,7 +811,10 @@ TSharedRef<SWidget> FAnimationMixerTrackEditor::BuildAddSectionSubMenu(TWeakView
 
 			NewSection->SetRowIndex(NewRowIndex);
 
-			MixerTrack->AddSection(*NewSection);
+			// Use the FSectionParameter overload to ensure lifecycle events are fired.
+			// The mixer track's OnSectionAddedImpl assigns sections to a layer, which is
+			// required for the section to appear in the UI.
+			static_cast<UMovieSceneTrack*>(MixerTrack)->AddSection(UMovieSceneTrack::FSectionParameter{*NewSection});
 
 			Sequencer->NotifyMovieSceneDataChanged(EMovieSceneDataChangeType::MovieSceneStructureItemAdded);
 		}
diff --git a/Engine/Plugins/Experimental/MovieScenePoseSearchTracks/Source/MovieScenePoseSearchTracksEditor/Private/TrackEditors/StitchAnimTrackEditor.cpp b/Engine/Plugins/Experimental/MovieScenePoseSearchTracks/Source/MovieScenePoseSearchTracksEditor/Private/TrackEditors/StitchAnimTrackEditor.cpp
--- a/Engine/Plugins/Experimental/MovieScenePoseSearchTracks/Source/MovieScenePoseSearchTracksEditor/Private/TrackEditors/StitchAnimTrackEditor.cpp
+++ b/Engine/Plugins/Experimental/MovieScenePoseSearchTracks/Source/MovieScenePoseSearchTracksEditor/Private/TrackEditors/StitchAnimTrackEditor.cpp
@@ -458,8 +458,8 @@ FKeyPropertyResult FStitchAnimTrackEditor::AddKeyInternal(FFrameNumber KeyTime,
 		// If a callback is provided (mixer case), create the section directly and let the callback handle adding it
 		if (InSectionCallback)
 		{
-			// Create section (outer will be set properly when callback adds it to track)
-			UMovieSceneStitchAnimSection* NewSection = NewObject<UMovieSceneStitchAnimSection>(GetTransientPackage(), NAME_None, RF_Transactional);
+			// Create section with the track as outer so MakeSectionInterface can find the mixer track
+			UMovieSceneStitchAnimSection* NewSection = NewObject<UMovieSceneStitchAnimSection>(Track, NAME_None, RF_Transactional);
 			NewSection->StitchDatabase = PoseSearchDatabase;
 
 			// Set up default section timing (duplicating what AddNewAnimationOnRow does)
@@ -535,7 +535,7 @@ void FStitchAnimTrackEditor::PopulateAddMixerItemMenu(FMenuBuilder& MenuBuilder,
 			{
 				// Use the specified row index
 				Section->SetRowIndex(RowIndex);
-				CommonTrack->AddSection(*Section);
+				static_cast<UMovieSceneTrack*>(CommonTrack)->AddSection(UMovieSceneTrack::FSectionParameter{*Section});
 			}
 			else
 			{
